WebSocketNetworking - making it available to multiple platforms (i.e. not just for HTML5)

- moved Engine/Platforms/HTML5/Plugins/WebSocketNetworking  to  Engine/Platforms/Plugins/WebSocketNetworking

#rnx

diff --git a/Engine/Platforms/HTML5/Plugins/WebSocketNetworking/README.md b/Engine/Platforms/HTML5/Plugins/WebSocketNetworking/README.md
deleted file mode 100644
diff --git a/Engine/Platforms/HTML5/Plugins/WebSocketNetworking/Source/WebSocketNetworking/Classes/WebSocketConnection.h b/Engine/Platforms/HTML5/Plugins/WebSocketNetworking/Source/WebSocketNetworking/Classes/WebSocketConnection.h
deleted file mode 100644
diff --git a/Engine/Platforms/HTML5/Plugins/WebSocketNetworking/Source/WebSocketNetworking/Classes/WebSocketNetDriver.h b/Engine/Platforms/HTML5/Plugins/WebSocketNetworking/Source/WebSocketNetworking/Classes/WebSocketNetDriver.h
deleted file mode 100644
diff --git a/Engine/Platforms/HTML5/Plugins/WebSocketNetworking/Source/WebSocketNetworking/Private/WebSocket.cpp b/Engine/Platforms/HTML5/Plugins/WebSocketNetworking/Source/WebSocketNetworking/Private/WebSocket.cpp
deleted file mode 100644
diff --git a/Engine/Platforms/HTML5/Plugins/WebSocketNetworking/Source/WebSocketNetworking/Private/WebSocket.h b/Engine/Platforms/HTML5/Plugins/WebSocketNetworking/Source/WebSocketNetworking/Private/WebSocket.h
deleted file mode 100644
diff --git a/Engine/Platforms/HTML5/Plugins/WebSocketNetworking/Source/WebSocketNetworking/Private/WebSocketNetDriver.cpp b/Engine/Platforms/HTML5/Plugins/WebSocketNetworking/Source/WebSocketNetworking/Private/WebSocketNetDriver.cpp
deleted file mode 100644
diff --git a/Engine/Platforms/HTML5/Plugins/WebSocketNetworking/Source/WebSocketNetworking/Private/WebSocketNetworkingPlugin.cpp b/Engine/Platforms/HTML5/Plugins/WebSocketNetworking/Source/WebSocketNetworking/Private/WebSocketNetworkingPlugin.cpp
deleted file mode 100644
diff --git a/Engine/Platforms/HTML5/Plugins/WebSocketNetworking/Source/WebSocketNetworking/Private/WebSocketNetworkingPrivate.h b/Engine/Platforms/HTML5/Plugins/WebSocketNetworking/Source/WebSocketNetworking/Private/WebSocketNetworkingPrivate.h
deleted file mode 100644
diff --git a/Engine/Platforms/HTML5/Plugins/WebSocketNetworking/Source/WebSocketNetworking/Private/WebSocketServer.cpp b/Engine/Platforms/HTML5/Plugins/WebSocketNetworking/Source/WebSocketNetworking/Private/WebSocketServer.cpp
deleted file mode 100644
diff --git a/Engine/Platforms/HTML5/Plugins/WebSocketNetworking/Source/WebSocketNetworking/Private/WebSocketServer.h b/Engine/Platforms/HTML5/Plugins/WebSocketNetworking/Source/WebSocketNetworking/Private/WebSocketServer.h
deleted file mode 100644
diff --git a/Engine/Platforms/HTML5/Plugins/WebSocketNetworking/Source/WebSocketNetworking/Private/WebsocketConnection.cpp b/Engine/Platforms/HTML5/Plugins/WebSocketNetworking/Source/WebSocketNetworking/Private/WebsocketConnection.cpp
deleted file mode 100644
diff --git a/Engine/Platforms/HTML5/Plugins/WebSocketNetworking/Source/WebSocketNetworking/WebSocketNetworking.Build.cs b/Engine/Platforms/HTML5/Plugins/WebSocketNetworking/Source/WebSocketNetworking/WebSocketNetworking.Build.cs
deleted file mode 100644
diff --git a/Engine/Platforms/HTML5/Plugins/WebSocketNetworking/WebSocketNetworking.uplugin b/Engine/Platforms/HTML5/Plugins/WebSocketNetworking/WebSocketNetworking.uplugin
deleted file mode 100644
diff --git a/Engine/Platforms/Plugins/WebSocketNetworking/Config/BaseWebSocketNetDriver.ini b/Engine/Platforms/Plugins/WebSocketNetworking/Config/BaseWebSocketNetDriver.ini
new file mode 100644
--- /dev/null
+++ b/Engine/Platforms/Plugins/WebSocketNetworking/Config/BaseWebSocketNetDriver.ini
@@ -0,0 +1,25 @@
+[/Script/Engine.GameEngine]
+!NetDriverDefinitions=ClearArray
++NetDriverDefinitions=(DefName="GameNetDriver", DriverClassName="/Script/WebSocketNetworking.WebSocketNetDriver", DriverClassNameFallback="/Script/WebSocketNetworking.WebSocketNetDriver")
+;[/Script/Engine.Engine]
+;+NetDriverDefinitions=(DefName="GameNetDriver", DriverClassName="/Script/WebSocketNetworking.WebSocketNetDriver", DriverClassNameFallback="/Script/WebSocketsNetworking.IpNetDriver")
+
+[/Script/WebSocketNetworking.WebSocketNetDriver]
+AllowPeerConnections=False
+AllowPeerVoice=False
+ConnectionTimeout=60.0
+InitialConnectTimeout=120.0
+RecentlyDisconnectedTrackingTime=180
+AckTimeout=10.0
+KeepAliveTime=20.2
+MaxClientRate=15000
+MaxInternetClientRate=10000
+RelevantTimeout=5.0
+SpawnPrioritySeconds=1.0
+ServerTravelPause=4.0
+NetServerMaxTickRate=30
+LanServerMaxTickRate=35
+WebSocketPort=8889
+NetConnectionClassName="/Script/WebSocketNetworking.WebSocketConnection"
+MaxPortCountToTry=512
+
diff --git a/Engine/Platforms/Plugins/WebSocketNetworking/Source/WebSocketNetworking/Classes/WebSocketConnection.h b/Engine/Platforms/Plugins/WebSocketNetworking/Source/WebSocketNetworking/Classes/WebSocketConnection.h
new file mode 100644
--- /dev/null
+++ b/Engine/Platforms/Plugins/WebSocketNetworking/Source/WebSocketNetworking/Classes/WebSocketConnection.h
@@ -0,0 +1,33 @@
+// Copyright 1998-2019 Epic Games, Inc. All Rights Reserved.
+
+#pragma once
+
+#include "Engine/NetConnection.h"
+#include "WebSocketConnection.generated.h"
+
+
+UCLASS(transient, config = Engine)
+class WEBSOCKETNETWORKING_API UWebSocketConnection : public UNetConnection
+{
+	GENERATED_UCLASS_BODY()
+
+	class FWebSocket*			WebSocket;
+
+	//~ Begin NetConnection Interface
+	virtual void InitBase(UNetDriver* InDriver, class FSocket* InSocket, const FURL& InURL, EConnectionState InState, int32 InMaxPacket = 0, int32 InPacketOverhead = 0) override;
+	virtual void InitRemoteConnection(UNetDriver* InDriver, class FSocket* InSocket, const FURL& InURL, const class FInternetAddr& InRemoteAddr, EConnectionState InState, int32 InMaxPacket = 0, int32 InPacketOverhead = 0) override;
+	virtual void InitLocalConnection(UNetDriver* InDriver, class FSocket* InSocket, const FURL& InURL, EConnectionState InState, int32 InMaxPacket = 0, int32 InPacketOverhead = 0) override;
+	virtual void LowLevelSend(void* Data, int32 CountBits, FOutPacketTraits& Traits) override;
+	FString LowLevelGetRemoteAddress(bool bAppendPort = false) override;
+	FString LowLevelDescribe() override;
+	virtual void Tick();
+	virtual void FinishDestroy();
+	virtual void ReceivedRawPacket(void* Data,int32 Count);
+	//~ End NetConnection Interface
+
+
+	void SetWebSocket(FWebSocket* InWebSocket);
+	FWebSocket* GetWebSocket();
+
+	bool bChallengeHandshake = false;
+};
diff --git a/Engine/Platforms/Plugins/WebSocketNetworking/Source/WebSocketNetworking/Classes/WebSocketNetDriver.h b/Engine/Platforms/Plugins/WebSocketNetworking/Source/WebSocketNetworking/Classes/WebSocketNetDriver.h
new file mode 100644
--- /dev/null
+++ b/Engine/Platforms/Plugins/WebSocketNetworking/Source/WebSocketNetworking/Classes/WebSocketNetDriver.h
@@ -0,0 +1,72 @@
+// Copyright 1998-2019 Epic Games, Inc. All Rights Reserved.
+//
+// websocket based implementation of the net driver
+//
+
+#pragma once
+
+#include "CoreMinimal.h"
+#include "Engine/NetDriver.h"
+#include "WebSocketNetDriver.generated.h"
+
+class FWebSocket;
+class FWebSocketServer;
+
+UCLASS(transient, config = Engine)
+class WEBSOCKETNETWORKING_API UWebSocketNetDriver : public UNetDriver
+{
+	GENERATED_UCLASS_BODY()
+
+	/** WebSocket server port*/
+	UPROPERTY(Config)
+	int32 WebSocketPort;
+
+	//~ Begin UNetDriver Interface.
+	virtual bool IsAvailable() const override;
+	virtual bool InitBase(bool bInitAsClient, FNetworkNotify* InNotify, const FURL& URL, bool bReuseAddressAndPort, FString& Error) override;
+	virtual bool InitConnect(FNetworkNotify* InNotify, const FURL& ConnectURL, FString& Error) override;
+	virtual bool InitListen(FNetworkNotify* InNotify, FURL& LocalURL, bool bReuseAddressAndPort, FString& Error) override;
+	virtual void TickDispatch(float DeltaTime) override;
+	virtual void LowLevelSend(TSharedPtr<const FInternetAddr> Address, void* Data, int32 CountBits, FOutPacketTraits& Traits) override;
+	virtual FString LowLevelGetNetworkNumber() override;
+	virtual void LowLevelDestroy() override;
+	virtual bool IsNetResourceValid(void) override;
+
+	// stub implementation because for websockets we don't use any underlying socket sub system.
+	virtual class ISocketSubsystem* GetSocketSubsystem() override;
+
+	//~ End UNetDriver Interface.
+
+	//~ Begin FExec Interface
+	virtual bool Exec(UWorld* InWorld, const TCHAR* Cmd, FOutputDevice& Ar = *GLog) override;
+	//~ End FExec Interface
+
+	/**
+	* Exec command handlers
+	*/
+	bool HandleSocketsCommand(const TCHAR* Cmd, FOutputDevice& Ar, UWorld* InWorld);
+
+	/** @return connection to server */
+	class UWebSocketConnection* GetServerConnection();
+
+	/************************************************************************/
+	/* FWebSocketServer														*/
+	/************************************************************************/
+
+	FWebSocketServer* WebSocketServer;
+
+	/******************************************************************************/
+	/* Callback Function for New Connection from a client is accepted by this server   */
+	/******************************************************************************/
+
+	void OnWebSocketClientConnected(FWebSocket*); // to the server.
+
+
+	/************************************************************************/
+	/* Callback Function for when this client Connects to the server				*/
+	/************************************************************************/
+
+	void OnWebSocketServerConnected(); // to the client.
+
+
+};
diff --git a/Engine/Platforms/Plugins/WebSocketNetworking/Source/WebSocketNetworking/Private/WebSocket.cpp b/Engine/Platforms/Plugins/WebSocketNetworking/Source/WebSocketNetworking/Private/WebSocket.cpp
new file mode 100644
--- /dev/null
+++ b/Engine/Platforms/Plugins/WebSocketNetworking/Source/WebSocketNetworking/Private/WebSocket.cpp
@@ -0,0 +1,502 @@
+// Copyright 1998-2019 Epic Games, Inc. All Rights Reserved.
+#include "WebSocket.h"
+#include "WebSocketNetworkingPrivate.h"
+#include "IPAddress.h"
+
+#if USE_LIBWEBSOCKET
+
+// Work around a conflict between a UI namespace defined by engine code and a typedef in OpenSSL
+#define UI UI_ST
+THIRD_PARTY_INCLUDES_START
+#include "libwebsockets.h"
+THIRD_PARTY_INCLUDES_END
+#undef UI
+
+static int unreal_networking_client(struct lws *wsi, enum lws_callback_reasons reason, void *user, void *in, size_t len);
+static void lws_debugLogS(int level, const char *line)
+{
+	UE_LOG(LogWebSocketNetworking, Log, TEXT("client: %s"), ANSI_TO_TCHAR(line));
+}
+
+#else
+
+#include <errno.h>
+#include <sys/types.h>
+#include <sys/socket.h>
+#include <netinet/in.h>
+#include <arpa/inet.h>
+#include <stdio.h>
+#include <stdlib.h>
+#include <string.h>
+#include <unistd.h>
+#include <fcntl.h>
+#include <sys/ioctl.h>
+#include <assert.h>
+#include <emscripten/emscripten.h>
+
+#endif
+
+FWebSocket::FWebSocket(
+		const FInternetAddr& ServerAddress
+)
+:IsServerSide(false)
+{
+
+#if USE_LIBWEBSOCKET
+
+#if !UE_BUILD_SHIPPING
+	lws_set_log_level(LLL_ERR | LLL_WARN | LLL_NOTICE | LLL_DEBUG | LLL_INFO, lws_debugLogS);
+#endif
+
+	Protocols = new lws_protocols[3];
+	FMemory::Memzero(Protocols, sizeof(lws_protocols) * 3);
+
+	Protocols[0].name = "binary";
+	Protocols[0].callback = unreal_networking_client;
+	Protocols[0].per_session_data_size = 0;
+	Protocols[0].rx_buffer_size = 10 * 1024 * 1024;
+
+	Protocols[1].name = nullptr;
+	Protocols[1].callback = nullptr;
+	Protocols[1].per_session_data_size = 0;
+
+	struct lws_context_creation_info Info;
+	memset(&Info, 0, sizeof Info);
+
+	Info.port = CONTEXT_PORT_NO_LISTEN;
+	Info.protocols = &Protocols[0];
+	Info.gid = -1;
+	Info.uid = -1;
+	Info.user = this;
+
+	//Info.options |= LWS_SERVER_OPTION_DO_SSL_GLOBAL_INIT;
+	Info.options |= LWS_SERVER_OPTION_DISABLE_IPV6;
+
+	Context = lws_create_context(&Info);
+
+	check(Context);
+
+	struct lws_client_connect_info ConnectInfo = {
+			Context, TCHAR_TO_ANSI(*ServerAddress.ToString(false)), ServerAddress.GetPort(), false, "/", TCHAR_TO_ANSI(*ServerAddress.ToString(false)), TCHAR_TO_ANSI(*ServerAddress.ToString(false)), Protocols[1].name, -1, this
+	};
+	Wsi = lws_client_connect_via_info(&ConnectInfo);
+	check(Wsi);
+
+#else // ! USE_LIBWEBSOCKET -- HTML5 uses BSD network API
+
+	SockFd = socket(PF_INET, SOCK_STREAM, IPPROTO_TCP);
+	if (SockFd == -1)
+	{
+		UE_LOG(LogWebSocketNetworking, Error, TEXT("Socket creation failed "));
+	}
+	else
+	{
+		UE_LOG(LogWebSocketNetworking, Warning, TEXT(" Socked %d created "), SockFd);
+	}
+
+	fcntl(SockFd, F_SETFL, O_NONBLOCK);
+
+#endif
+
+	memset(&RemoteAddr, 0, sizeof(RemoteAddr));
+
+	// Windows XP does not have support for inet_pton
+#if PLATFORM_WINDOWS && _WIN32_WINNT <= 0x0502
+	int32 SizeOfRemoteAddr = sizeof(RemoteAddr);
+
+	// Force ServerAddress into non-const array. API doesn't modify contents but old API still requires non-const string
+	if (WSAStringToAddress(ServerAddress.ToString(false).GetCharArray().GetData(), AF_INET, NULL, (sockaddr*)&RemoteAddr, &SizeOfRemoteAddr) != 0)
+	{
+		UE_LOG(LogWebSocketNetworking, Warning, TEXT("WSAStringToAddress failed "));
+		return;
+	}
+#else
+	if (inet_pton(AF_INET, TCHAR_TO_ANSI(*ServerAddress.ToString(false)), &RemoteAddr.sin_addr) != 1)
+	{
+		UE_LOG(LogWebSocketNetworking, Warning, TEXT("inet_pton failed to %s"), *ServerAddress.ToString(false));
+		return;
+	}
+#endif
+
+	RemoteAddr.sin_family = AF_INET;
+	RemoteAddr.sin_port = htons(ServerAddress.GetPort());
+
+#if !USE_LIBWEBSOCKET // HTML5 uses BSD network API
+	int Ret = connect(SockFd, (struct sockaddr *)&RemoteAddr, sizeof(RemoteAddr));
+	UE_LOG(LogWebSocketNetworking, Warning, TEXT(" Connect socket returned %d to %s. Error Code: %d"), Ret, *ServerAddress.ToString(false), ((Ret != 0) ? errno : 0));
+#endif
+}
+
+#if USE_LIBWEBSOCKET
+FWebSocket::FWebSocket(WebSocketInternalContext* InContext, WebSocketInternal* InWsi)
+	: Context(InContext)
+	, Wsi(InWsi)
+	, Protocols(nullptr)
+	, IsServerSide(true)
+{
+	int sock = lws_get_socket_fd(Wsi);
+	socklen_t len = sizeof RemoteAddr;
+	getpeername(sock, (struct sockaddr*)&RemoteAddr, &len);
+}
+#endif
+
+bool FWebSocket::Send(uint8* Data, uint32 Size)
+{
+	TArray<uint8> Buffer;
+
+#if USE_LIBWEBSOCKET
+	Buffer.AddDefaulted(LWS_PRE); // Reserve space for WS header data
+#endif
+
+	Buffer.Append((uint8*)&Size, sizeof (uint32)); // insert size.
+	Buffer.Append((uint8*)Data, Size);
+	OutgoingBuffer.Add(Buffer);
+
+	return true;
+}
+
+void FWebSocket::SetRecieveCallBack(FWebSocketPacketRecievedCallBack CallBack)
+{
+	RecievedCallBack = CallBack;
+}
+
+FString FWebSocket::RemoteEndPoint(bool bAppendPort)
+{
+	// Windows XP does not have support for inet_ntop
+#if PLATFORM_WINDOWS && _WIN32_WINNT <= 0x0502
+	TCHAR Buffer[INET_ADDRSTRLEN];
+	::DWORD BufferSize = INET_ADDRSTRLEN;
+	FString remote = "";
+	sockaddr_in AddressToConvert = RemoteAddr;
+	if (!bAppendPort)
+	{
+		AddressToConvert.sin_port = 0;
+	}
+	if (WSAAddressToString((sockaddr*)&AddressToConvert, sizeof(AddressToConvert), NULL, Buffer, &BufferSize) == 0)
+	{
+		remote = Buffer;
+	}
+#else
+	ANSICHAR Buffer[INET_ADDRSTRLEN];
+	FString remote(ANSI_TO_TCHAR(inet_ntop(AF_INET, &RemoteAddr.sin_addr, Buffer, INET_ADDRSTRLEN)));
+	if ( bAppendPort )
+	{
+		remote += FString::Printf(TEXT(":%i"), ntohs(RemoteAddr.sin_port));
+	}
+#endif
+
+	return remote;
+}
+
+FString FWebSocket::LocalEndPoint(bool bAppendPort)
+{
+#if USE_LIBWEBSOCKET
+	int sock = lws_get_socket_fd(Wsi);
+	struct sockaddr_in addr;
+	socklen_t len = sizeof addr;
+	getsockname(sock, (struct sockaddr*)&addr, &len);
+
+	// Windows XP does not have support for inet_ntop
+#if PLATFORM_WINDOWS && _WIN32_WINNT <= 0x0502
+	TCHAR Buffer[INET_ADDRSTRLEN];
+	::DWORD BufferSize = INET_ADDRSTRLEN;
+	FString remote = "";
+	if (!bAppendPort)
+	{
+		addr.sin_port = 0;
+	}
+	if (WSAAddressToString((sockaddr*)&addr, sizeof(addr), NULL, Buffer, &BufferSize) == 0)
+	{
+		remote = Buffer;
+	}
+#else
+	ANSICHAR Buffer[INET_ADDRSTRLEN];
+	FString remote(ANSI_TO_TCHAR(inet_ntop(AF_INET, &addr.sin_addr, Buffer, INET_ADDRSTRLEN)));
+	if ( bAppendPort )
+	{
+		remote += FString::Printf(TEXT(":%i"), ntohs(addr.sin_port));
+	}
+#endif
+
+	return remote;
+#else // ! USE_LIBWEBSOCKET -- HTML5 uses BSD network API
+	// NOTE: there's no way to get this info from browsers...
+	// return generic localhost without port #
+	return FString(TEXT("127.0.0.1"));
+#endif
+}
+
+void FWebSocket::Tick()
+{
+	HandlePacket();
+}
+
+void FWebSocket::HandlePacket()
+{
+#if USE_LIBWEBSOCKET
+
+	lws_service(Context, 0);
+	if (!IsServerSide)
+		lws_callback_on_writable_all_protocol(Context, &Protocols[0]);
+
+#else // ! USE_LIBWEBSOCKET -- HTML5 uses BSD network API
+
+	fd_set Fdr;
+	fd_set Fdw;
+	int Res;
+
+	// make sure that server.fd is ready to read / write
+	FD_ZERO(&Fdr);
+	FD_ZERO(&Fdw);
+	FD_SET(SockFd, &Fdr);
+	FD_SET(SockFd, &Fdw);
+	Res = select(64, &Fdr, &Fdw, NULL, NULL);
+
+	if (Res == -1) {
+		UE_LOG(LogWebSocketNetworking, Warning, TEXT("Select Failed!"));
+		return;
+	}
+
+	if (FD_ISSET(SockFd, &Fdr)) {
+		// we can read!
+		OnRawRecieve(NULL, NULL);
+	}
+
+	if (FD_ISSET(SockFd, &Fdw)) {
+		// we can write
+		OnRawWebSocketWritable(NULL);
+	}
+
+#endif
+}
+
+void FWebSocket::Flush()
+{
+	auto PendingMesssages = OutgoingBuffer.Num();
+	while (OutgoingBuffer.Num() > 0 && !IsServerSide)
+	{
+#if USE_LIBWEBSOCKET
+		if (Protocols)
+		{
+			lws_callback_on_writable_all_protocol(Context, &Protocols[0]);
+		}
+		else
+		{
+			lws_callback_on_writable(Wsi);
+		}
+#endif
+		HandlePacket();
+		if (PendingMesssages >= OutgoingBuffer.Num())
+		{
+			UE_LOG(LogWebSocketNetworking, Warning, TEXT("Unable to flush all of OutgoingBuffer in FWebSocket."));
+			break;
+		}
+	};
+}
+
+TArray<uint8> FWebSocket::GetRawRemoteAddr(int32& OutPort)
+{
+	OutPort = ntohs(RemoteAddr.sin_port);
+	TArray<uint8> RawBuffer;
+	uint32 IntAddr = RemoteAddr.sin_addr.s_addr;
+	RawBuffer.Add((IntAddr >> 0) & 0xFF);
+	RawBuffer.Add((IntAddr >> 8) & 0xFF);
+	RawBuffer.Add((IntAddr >> 16) & 0xFF);
+	RawBuffer.Add((IntAddr >> 24) & 0xFF);
+
+	return RawBuffer;
+}
+
+void FWebSocket::SetConnectedCallBack(FWebSocketInfoCallBack CallBack)
+{
+	ConnectedCallBack = CallBack;
+}
+
+void FWebSocket::SetErrorCallBack(FWebSocketInfoCallBack CallBack)
+{
+	ErrorCallBack = CallBack;
+}
+
+void FWebSocket::OnRawRecieve(void* Data, uint32 Size)
+{
+#if USE_LIBWEBSOCKET
+
+	RecievedBuffer.Append((uint8*)Data, Size); // consumes all of Data
+	while (RecievedBuffer.Num() > sizeof(uint32))
+	{
+		uint32 BytesToBeRead = *(uint32*)RecievedBuffer.GetData();
+		if (BytesToBeRead <= ((uint32)RecievedBuffer.Num() - sizeof(uint32)))
+		{
+			RecievedCallBack.ExecuteIfBound((void*)((uint8*)RecievedBuffer.GetData() + sizeof(uint32)), BytesToBeRead);
+			RecievedBuffer.RemoveAt(0, sizeof(uint32) + BytesToBeRead );
+		}
+		else
+		{
+			break;
+		}
+	}
+
+#else // ! USE_LIBWEBSOCKET -- HTML5 uses BSD network API
+
+	// browser was crashing when using RecievedBuffer...
+
+	check(Data == NULL); // jic this is not obvious, Data will be resigned to Buffer below
+
+	uint8 Buffer[1024]; // should be at MAX PACKET SIZE.
+	Size = recv(SockFd, Buffer, sizeof(Buffer), 0);
+//	while ( Size > sizeof(uint32) )
+	if ( Size > sizeof(uint32) )
+	{
+		Data = (void*)Buffer;
+		uint32 BytesToBeRead = *(uint32*)Data;
+		uint32 BytesLeft = Size;
+		while ( ( BytesToBeRead > 0 ) && ( BytesToBeRead < BytesLeft ) )
+		{
+			Data = (void*)((uint8*)Data + sizeof(uint32));
+			RecievedCallBack.ExecuteIfBound(Data, BytesToBeRead);
+
+			// "RecievedBuffer.RemoveAt()"
+			Data = (void*)((uint8*)Data + BytesToBeRead);
+			BytesLeft -= BytesToBeRead + sizeof(uint32);
+			if ( (uint8*)Data >= (Buffer+Size)  // hard cap
+				||	BytesLeft == 0	)           // soft limit
+			{
+				break;
+			}
+			BytesToBeRead = *(uint32*)Data;
+		}
+//		Size = recv(SockFd, Buffer, sizeof(Buffer), 0);
+	}
+
+#endif
+}
+
+void FWebSocket::OnRawWebSocketWritable(WebSocketInternal* wsi)
+{
+	if (OutgoingBuffer.Num() == 0)
+		return;
+
+	TArray <uint8>& Packet = OutgoingBuffer[0];
+
+#if USE_LIBWEBSOCKET
+
+	uint32 TotalDataSize = Packet.Num() - LWS_PRE;
+	uint32 DataToSend = TotalDataSize;
+	while (DataToSend)
+	{
+		int Sent = lws_write(Wsi, Packet.GetData() + LWS_PRE + (DataToSend-TotalDataSize), DataToSend, (lws_write_protocol)LWS_WRITE_BINARY);
+		if (Sent < 0)
+		{
+			ErrorCallBack.ExecuteIfBound();
+			return;
+		}
+		if ((uint32)Sent < DataToSend)
+		{
+			UE_LOG(LogWebSocketNetworking, Warning, TEXT("Could not write all '%d' bytes to socket"), DataToSend);
+		}
+		DataToSend-=Sent;
+	}
+
+	check(Wsi == wsi);
+
+#else // ! USE_LIBWEBSOCKET -- HTML5 uses BSD network API
+
+	uint32 TotalDataSize = Packet.Num();
+	uint32 DataToSend = TotalDataSize;
+	while (DataToSend)
+	{
+		// send actual data in one go.
+		int Result = send(SockFd, Packet.GetData()+(DataToSend-TotalDataSize),DataToSend, 0);
+		if (Result == -1)
+		{
+			// we are caught with our pants down. fail.
+			UE_LOG(LogWebSocketNetworking, Error, TEXT("Could not write %d bytes"), Packet.Num());
+			ErrorCallBack.ExecuteIfBound();
+			return;
+		}
+		UE_CLOG((uint32)Result < DataToSend, LogWebSocketNetworking, Warning, TEXT("Could not write all '%d' bytes to socket"), DataToSend);
+		DataToSend-=Result;
+	}
+
+#endif
+
+	// this is very inefficient we need a constant size circular buffer to efficiently not do unnecessary allocations/deallocations.
+	OutgoingBuffer.RemoveAt(0);
+
+}
+
+FWebSocket::~FWebSocket()
+{
+	RecievedCallBack.Unbind();
+
+#if USE_LIBWEBSOCKET
+
+	Flush();
+
+	if ( !IsServerSide)
+	{
+		lws_context_destroy(Context);
+		Context = NULL;
+		delete Protocols;
+		Protocols = NULL;
+	}
+
+#else // ! USE_LIBWEBSOCKET -- HTML5 uses BSD network API
+
+	close(SockFd);
+
+#endif
+
+}
+
+#if USE_LIBWEBSOCKET
+static int unreal_networking_client(
+		struct lws *Wsi,
+		enum lws_callback_reasons Reason,
+		void *User,
+		void *In,
+		size_t Len)
+{
+	struct lws_context *Context = lws_get_context(Wsi);
+	FWebSocket* Socket = (FWebSocket*)lws_context_user(Context);
+	switch (Reason)
+	{
+	case LWS_CALLBACK_CLIENT_ESTABLISHED:
+		{
+			Socket->ConnectedCallBack.ExecuteIfBound();
+			lws_set_timeout(Wsi, NO_PENDING_TIMEOUT, 0);
+			check(Socket->Wsi == Wsi);
+		}
+		break;
+	case LWS_CALLBACK_CLIENT_CONNECTION_ERROR:
+		{
+			Socket->ErrorCallBack.ExecuteIfBound();
+			return -1;
+		}
+		break;
+	case LWS_CALLBACK_CLIENT_RECEIVE:
+		{
+			// push it on the socket.
+			Socket->OnRawRecieve(In, (uint32)Len);
+			check(Socket->Wsi == Wsi);
+			lws_set_timeout(Wsi, NO_PENDING_TIMEOUT, 0);
+			break;
+		}
+	case LWS_CALLBACK_CLIENT_WRITEABLE:
+		{
+			check(Socket->Wsi == Wsi);
+			Socket->OnRawWebSocketWritable(Wsi);
+			lws_callback_on_writable(Wsi);
+			lws_set_timeout(Wsi, NO_PENDING_TIMEOUT, 0);
+			break;
+		}
+	case LWS_CALLBACK_CLIENT_CLOSED:
+		{
+			Socket->ErrorCallBack.ExecuteIfBound();
+			return -1;
+		}
+	}
+
+	return 0;
+}
+#endif
diff --git a/Engine/Platforms/Plugins/WebSocketNetworking/Source/WebSocketNetworking/Private/WebSocket.h b/Engine/Platforms/Plugins/WebSocketNetworking/Source/WebSocketNetworking/Private/WebSocket.h
new file mode 100644
--- /dev/null
+++ b/Engine/Platforms/Plugins/WebSocketNetworking/Source/WebSocketNetworking/Private/WebSocket.h
@@ -0,0 +1,86 @@
+// Copyright 1998-2019 Epic Games, Inc. All Rights Reserved.
+//
+// libwebsocket client wrapper.
+//
+#pragma  once
+#include "WebSocketNetworkingPrivate.h"
+#if USE_LIBWEBSOCKET
+#include "Runtime/Sockets/Private/BSDSockets/SocketSubsystemBSD.h"
+#else
+#include <netinet/in.h>
+#endif
+
+class FWebSocket
+{
+
+public:
+
+	// Initialize as client side socket.
+	FWebSocket(const FInternetAddr& ServerAddress);
+
+#if USE_LIBWEBSOCKET
+	// Initialize as server side socket.
+	FWebSocket(WebSocketInternalContext* InContext, WebSocketInternal* Wsi);
+#endif
+
+	// clean up.
+	~FWebSocket();
+
+	/************************************************************************/
+	/* Set various callbacks for Socket Events                              */
+	/************************************************************************/
+	void SetConnectedCallBack(FWebSocketInfoCallBack CallBack);
+	void SetErrorCallBack(FWebSocketInfoCallBack CallBack);
+	void SetRecieveCallBack(FWebSocketPacketRecievedCallBack CallBack);
+
+	/** Send raw data to remote end point. */
+	bool Send(uint8* Data, uint32 Size);
+
+	/** service libwebsocket.			   */
+	void Tick();
+	/** service libwebsocket until outgoing buffer is empty */
+	void Flush();
+
+	/** Helper functions to describe end points. */
+	TArray<uint8> GetRawRemoteAddr(int32& OutPort);
+	FString RemoteEndPoint(bool bAppendPort);
+	FString LocalEndPoint(bool bAppendPort);
+	struct sockaddr_in* GetRemoteAddr() { return &RemoteAddr; }
+
+// this was made public because of cross-platform build issues
+public:
+
+	void HandlePacket();
+	void OnRawRecieve(void* Data, uint32 Size);
+	void OnRawWebSocketWritable(WebSocketInternal* wsi);
+
+	/************************************************************************/
+	/*	Various Socket callbacks											*/
+	/************************************************************************/
+	FWebSocketPacketRecievedCallBack  RecievedCallBack;
+	FWebSocketInfoCallBack ConnectedCallBack;
+	FWebSocketInfoCallBack ErrorCallBack;
+
+	/**  Recv and Send Buffers, serviced during the Tick */
+	TArray<uint8> RecievedBuffer;
+	TArray<TArray<uint8>> OutgoingBuffer;
+
+#if USE_LIBWEBSOCKET
+	/** libwebsocket internal context*/
+	WebSocketInternalContext* Context;
+
+	/** libwebsocket web socket */
+	WebSocketInternal* Wsi;
+
+	/** libwebsocket Protocols that can be serviced by this implemenation*/
+	WebSocketInternalProtocol* Protocols;
+#else // ! USE_LIBWEBSOCKET -- HTML5 uses BSD network API
+	int SockFd;
+#endif
+	struct sockaddr_in RemoteAddr;
+
+	/** Server side socket or client side*/
+	bool IsServerSide;
+
+	friend class FWebSocketServer;
+};
diff --git a/Engine/Platforms/Plugins/WebSocketNetworking/Source/WebSocketNetworking/Private/WebSocketNetDriver.cpp b/Engine/Platforms/Plugins/WebSocketNetworking/Source/WebSocketNetworking/Private/WebSocketNetDriver.cpp
new file mode 100644
--- /dev/null
+++ b/Engine/Platforms/Plugins/WebSocketNetworking/Source/WebSocketNetworking/Private/WebSocketNetDriver.cpp
@@ -0,0 +1,277 @@
+// Copyright 1998-2019 Epic Games, Inc. All Rights Reserved.
+
+/*=============================================================================
+Unreal WebSocket network driver.
+=============================================================================*/
+
+#include "WebSocketNetDriver.h"
+#include "WebSocketNetworkingPrivate.h"
+#include "Engine/Channel.h"
+#include "Engine/PendingNetGame.h"
+
+#include "IPAddress.h"
+#include "Sockets.h"
+
+#include "WebSocketConnection.h"
+#include "WebSocketServer.h"
+#include "WebSocket.h"
+
+#include "Engine/ChildConnection.h"
+#include "Misc/CommandLine.h"
+
+/*-----------------------------------------------------------------------------
+Declarations.
+-----------------------------------------------------------------------------*/
+
+/** Size of the network recv buffer */
+#define NETWORK_MAX_PACKET (576)
+
+UWebSocketNetDriver::UWebSocketNetDriver(const FObjectInitializer& ObjectInitializer)
+	: Super(ObjectInitializer)
+{
+}
+
+bool UWebSocketNetDriver::IsAvailable() const
+{
+	// WebSocket driver always valid for now
+	return true;
+}
+
+ISocketSubsystem* UWebSocketNetDriver::GetSocketSubsystem()
+{
+	return ISocketSubsystem::Get();
+}
+
+bool UWebSocketNetDriver::InitBase(bool bInitAsClient, FNetworkNotify* InNotify, const FURL& URL, bool bReuseAddressAndPort, FString& Error)
+{
+	if (!Super::InitBase(bInitAsClient, InNotify, URL, bReuseAddressAndPort, Error))
+	{
+		return false;
+	}
+
+	return true;
+}
+
+bool UWebSocketNetDriver::InitConnect(FNetworkNotify* InNotify, const FURL& ConnectURL, FString& Error)
+{
+	if (!InitBase(true, InNotify, ConnectURL, false, Error))
+	{
+		return false;
+	}
+
+	// Create new connection.
+	ServerConnection = NewObject<UWebSocketConnection>(NetConnectionClass);
+
+	TSharedRef<FInternetAddr> InternetAddr = GetSocketSubsystem()->CreateInternetAddr();
+	bool Ok;
+	InternetAddr->SetIp(*ConnectURL.Host, Ok);
+	InternetAddr->SetPort(WebSocketPort);
+
+	FWebSocket* WebSocket = new FWebSocket(*InternetAddr);
+	UWebSocketConnection* Connection = (UWebSocketConnection*)ServerConnection;
+	Connection->SetWebSocket(WebSocket);
+
+	FWebSocketPacketRecievedCallBack CallBack;
+	CallBack.BindUObject(Connection, &UWebSocketConnection::ReceivedRawPacket);
+	WebSocket->SetRecieveCallBack(CallBack);
+
+	FWebSocketInfoCallBack  ConnectedCallBack;
+	ConnectedCallBack.BindUObject(this, &UWebSocketNetDriver::OnWebSocketServerConnected);
+	WebSocket->SetConnectedCallBack(ConnectedCallBack);
+
+	ServerConnection->InitLocalConnection(this, NULL, ConnectURL, USOCK_Pending);
+
+	// Create channel zero.
+	CreateInitialClientChannels();
+
+	return true;
+}
+
+bool UWebSocketNetDriver::InitListen(FNetworkNotify* InNotify, FURL& LocalURL, bool bReuseAddressAndPort, FString& Error)
+{
+	if (!InitBase(false, InNotify, LocalURL, bReuseAddressAndPort, Error))
+	{
+		return false;
+	}
+
+	InitConnectionlessHandler();
+
+	WebSocketServer = new class FWebSocketServer();
+
+	FWebSocketClientConnectedCallBack CallBack;
+	CallBack.BindUObject(this, &UWebSocketNetDriver::OnWebSocketClientConnected);
+
+	if(!WebSocketServer->Init(WebSocketPort, CallBack))
+		return false;
+
+	WebSocketServer->Tick();
+	LocalURL.Port = WebSocketPort;
+	LocalAddr = GetSocketSubsystem()->GetLocalBindAddr(*GLog);
+	LocalAddr->SetPort(WebSocketPort);
+	UE_LOG(LogWebSocketNetworking, Log, TEXT("%s WebSocketNetDriver listening on port %i"), *GetDescription(), LocalURL.Port);
+
+	// server has no server connection.
+	ServerConnection = NULL;
+	return true;
+}
+
+void UWebSocketNetDriver::TickDispatch(float DeltaTime)
+{
+	Super::TickDispatch(DeltaTime);
+
+	if (WebSocketServer)
+		WebSocketServer->Tick();
+}
+
+void UWebSocketNetDriver::LowLevelSend(TSharedPtr<const FInternetAddr> Address, void* Data, int32 CountBits, FOutPacketTraits& Traits)
+{
+	if (Address.IsValid() && Address->IsValid())
+	{
+		const uint8* DataToSend = reinterpret_cast<uint8*>(Data);
+
+		if (ConnectionlessHandler.IsValid())
+		{
+			const ProcessedPacket ProcessedData =
+					ConnectionlessHandler->OutgoingConnectionless(Address, (uint8*)DataToSend, CountBits, Traits);
+
+			if (!ProcessedData.bError)
+			{
+				DataToSend = ProcessedData.Data;
+				CountBits = ProcessedData.CountBits;
+			}
+			else
+			{
+				CountBits = 0;
+			}
+		}
+
+
+		// connectionless websockets do not exist (yet)
+		// scan though existing connections
+		if (CountBits > 0)
+		{
+			for (int32 i = 0; i<ClientConnections.Num(); ++i)
+			{
+				UWebSocketConnection* Connection = (UWebSocketConnection*)ClientConnections[i];
+				if (Connection && ( Connection->LowLevelGetRemoteAddress(true) == Address->ToString(true) ) )
+				{
+					Connection->GetWebSocket()->Send((uint8*)DataToSend, FMath::DivideAndRoundUp(CountBits, 8));
+					break;
+				}
+			}
+		}
+	}
+	else
+	{
+		UE_LOG(LogNet, Warning, TEXT("UWebSocketNetDriver::LowLevelSend: Invalid send address '%s'"), *Address->ToString(false));
+	}
+}
+
+FString UWebSocketNetDriver::LowLevelGetNetworkNumber()
+{
+	return WebSocketServer != nullptr ? WebSocketServer->Info() : FString(TEXT(""));
+}
+
+void UWebSocketNetDriver::LowLevelDestroy()
+{
+	Super::LowLevelDestroy();
+	delete WebSocketServer;
+	WebSocketServer = nullptr;
+}
+
+bool UWebSocketNetDriver::HandleSocketsCommand(const TCHAR* Cmd, FOutputDevice& Ar, UWorld* InWorld)
+{
+	Ar.Logf(TEXT(""));
+	if (WebSocketServer != NULL)
+	{
+		Ar.Logf(TEXT("Running WebSocket Server %s"), *WebSocketServer->Info());
+	}
+	else
+	{
+		check(GetServerConnection());
+		Ar.Logf(TEXT("WebSocket client's EndPoint %s"), *(GetServerConnection()->WebSocket->RemoteEndPoint(true)));
+	}
+	return UNetDriver::Exec(InWorld, TEXT("SOCKETS"), Ar);
+}
+
+bool UWebSocketNetDriver::Exec(UWorld* InWorld, const TCHAR* Cmd, FOutputDevice& Ar)
+{
+	if (FParse::Command(&Cmd, TEXT("SOCKETS")))
+	{
+		return HandleSocketsCommand(Cmd, Ar, InWorld);
+	}
+	return UNetDriver::Exec(InWorld, Cmd, Ar);
+}
+
+UWebSocketConnection* UWebSocketNetDriver::GetServerConnection()
+{
+	return (UWebSocketConnection*)ServerConnection;
+}
+
+void UWebSocketNetDriver::OnWebSocketClientConnected(FWebSocket* ClientWebSocket)
+{
+	// Determine if allowing for client/server connections
+	const bool bAcceptingConnection = Notify->NotifyAcceptingConnection() == EAcceptConnection::Accept;
+	if (bAcceptingConnection)
+	{
+
+		UWebSocketConnection* Connection = NewObject<UWebSocketConnection>(NetConnectionClass);
+		check(Connection);
+
+		TSharedRef<FInternetAddr> InternetAddr = GetSocketSubsystem()->CreateInternetAddr();
+		int32 AddressPort = 0;
+
+		InternetAddr->SetRawIp(ClientWebSocket->GetRawRemoteAddr(AddressPort));
+		InternetAddr->SetPort(0);
+		Connection->SetWebSocket(ClientWebSocket);
+		Connection->InitRemoteConnection(this, NULL, FURL(), *InternetAddr, USOCK_Open);
+		
+		Notify->NotifyAcceptedConnection(Connection);
+
+		AddClientConnection(Connection);
+
+		FWebSocketPacketRecievedCallBack CallBack;
+		CallBack.BindUObject(Connection, &UWebSocketConnection::ReceivedRawPacket);
+		if (ConnectionlessHandler.IsValid() && StatelessConnectComponent.IsValid())
+		{
+			Connection->bChallengeHandshake = true;
+		}
+#if !UE_BUILD_SHIPPING
+		else if (FParse::Param(FCommandLine::Get(), TEXT("NoPacketHandler")))
+		{
+			UE_LOG(LogNet, Log, TEXT("Accepting connection without handshake, due to '-NoPacketHandler'."))
+		}
+#endif
+		else
+		{
+			UE_LOG(LogNet, Log,
+					TEXT("Invalid ConnectionlessHandler (%i) or StatelessConnectComponent (%i); can't accept connections."),
+					(int32)(ConnectionlessHandler.IsValid()), (int32)(StatelessConnectComponent.IsValid()));
+		}
+		ClientWebSocket->SetRecieveCallBack(CallBack);
+
+		UE_LOG(LogWebSocketNetworking, Log, TEXT(" WebSocket server running on %s Accepted Connection from %s "), *WebSocketServer->Info(),*ClientWebSocket->RemoteEndPoint(true));
+	}
+}
+
+bool UWebSocketNetDriver::IsNetResourceValid(void)
+{
+	if (	(WebSocketServer && !ServerConnection)//  Server
+		||	(!WebSocketServer && ServerConnection) // client
+		)
+	{
+		return true;
+	}
+		
+	return false;
+}
+
+// Just logging, not yet attached to html5 clients.
+void UWebSocketNetDriver::OnWebSocketServerConnected()
+{
+	check(GetServerConnection());
+	UE_LOG(LogWebSocketNetworking, Log, TEXT(" %s WebSocket Client %s connected to server %s "), *GetDescription(),
+		*GetServerConnection()->WebSocket->LocalEndPoint(true),
+		*GetServerConnection()->WebSocket->RemoteEndPoint(true));
+}
+
diff --git a/Engine/Platforms/Plugins/WebSocketNetworking/Source/WebSocketNetworking/Private/WebSocketNetworkingPlugin.cpp b/Engine/Platforms/Plugins/WebSocketNetworking/Source/WebSocketNetworking/Private/WebSocketNetworkingPlugin.cpp
new file mode 100644
--- /dev/null
+++ b/Engine/Platforms/Plugins/WebSocketNetworking/Source/WebSocketNetworking/Private/WebSocketNetworkingPlugin.cpp
@@ -0,0 +1,26 @@
+// Copyright 1998-2019 Epic Games, Inc. All Rights Reserved.
+
+#include "WebSocketNetworkingPrivate.h"
+
+
+class FWebSocketNetworkingPlugin : public IModuleInterface
+{
+	/** IModuleInterface implementation */
+	virtual void StartupModule() override;
+	virtual void ShutdownModule() override;
+};
+
+IMPLEMENT_MODULE(FWebSocketNetworkingPlugin, WebSocketNetworking)
+
+void FWebSocketNetworkingPlugin::StartupModule()
+{
+}
+
+
+void FWebSocketNetworkingPlugin::ShutdownModule()
+{
+}
+
+DEFINE_LOG_CATEGORY(LogWebSocketNetworking);
+
+
diff --git a/Engine/Platforms/Plugins/WebSocketNetworking/Source/WebSocketNetworking/Private/WebSocketNetworkingPrivate.h b/Engine/Platforms/Plugins/WebSocketNetworking/Source/WebSocketNetworking/Private/WebSocketNetworkingPrivate.h
new file mode 100644
--- /dev/null
+++ b/Engine/Platforms/Plugins/WebSocketNetworking/Source/WebSocketNetworking/Private/WebSocketNetworkingPrivate.h
@@ -0,0 +1,20 @@
+// Copyright 1998-2019 Epic Games, Inc. All Rights Reserved.
+#pragma once
+
+#include "CoreMinimal.h"
+#include "Sockets.h"
+#include "SocketSubsystem.h"
+#include "Modules/ModuleManager.h"
+
+class FWebSocket; 
+class FWebSocketServer; 
+
+typedef struct lws_context WebSocketInternalContext;
+typedef struct lws WebSocketInternal;
+typedef struct lws_protocols WebSocketInternalProtocol;
+
+DECLARE_DELEGATE_TwoParams(FWebSocketPacketRecievedCallBack, void* /*Data*/, int32 /*Data Size*/);
+DECLARE_DELEGATE_OneParam(FWebSocketClientConnectedCallBack, FWebSocket* /*Socket*/);
+DECLARE_DELEGATE(FWebSocketInfoCallBack); 
+
+DECLARE_LOG_CATEGORY_EXTERN(LogWebSocketNetworking, Warning, All);
diff --git a/Engine/Platforms/Plugins/WebSocketNetworking/Source/WebSocketNetworking/Private/WebSocketServer.cpp b/Engine/Platforms/Plugins/WebSocketNetworking/Source/WebSocketNetworking/Private/WebSocketServer.cpp
new file mode 100644
--- /dev/null
+++ b/Engine/Platforms/Plugins/WebSocketNetworking/Source/WebSocketNetworking/Private/WebSocketServer.cpp
@@ -0,0 +1,178 @@
+// Copyright 1998-2019 Epic Games, Inc. All Rights Reserved.
+
+#include "WebSocketServer.h"
+#include "WebSocketNetworkingPrivate.h"
+#include "WebSocket.h"
+
+#if USE_LIBWEBSOCKET
+// Work around a conflict between a UI namespace defined by engine code and a typedef in OpenSSL
+#define UI UI_ST
+THIRD_PARTY_INCLUDES_START
+#include "libwebsockets.h"
+THIRD_PARTY_INCLUDES_END
+#undef UI
+#endif
+
+// a object of this type is associated by libwebsocket to every connected session.
+struct PerSessionDataServer
+{
+	FWebSocket *Socket; // each session is actually a socket to a client
+};
+
+
+#if USE_LIBWEBSOCKET
+// real networking handler.
+static int unreal_networking_server(struct lws *wsi, enum lws_callback_reasons reason, void *user, void *in, size_t len);
+
+#if !UE_BUILD_SHIPPING
+	inline void lws_debugLog(int level, const char *line)
+	{
+		UE_LOG(LogWebSocketNetworking, Log, TEXT("websocket server: %s"), ANSI_TO_TCHAR(line));
+	}
+#endif
+#endif
+
+bool FWebSocketServer::Init(uint32 Port, FWebSocketClientConnectedCallBack CallBack)
+{
+#if USE_LIBWEBSOCKET
+#if !UE_BUILD_SHIPPING
+	// setup log level.
+	lws_set_log_level(LLL_ERR | LLL_WARN | LLL_NOTICE | LLL_DEBUG | LLL_INFO, lws_debugLog);
+#endif
+
+	Protocols = new lws_protocols[3];
+	FMemory::Memzero(Protocols, sizeof(lws_protocols) * 3);
+
+	Protocols[0].name = "binary";
+	Protocols[0].callback = unreal_networking_server;
+	Protocols[0].per_session_data_size = sizeof(PerSessionDataServer);
+	Protocols[0].rx_buffer_size = 10 * 1024 * 1024;
+
+	Protocols[1].name = nullptr;
+	Protocols[1].callback = nullptr;
+	Protocols[1].per_session_data_size = 0;
+
+	struct lws_context_creation_info Info;
+	memset(&Info, 0, sizeof(lws_context_creation_info));
+	// look up libwebsockets.h for details.
+	Info.port = Port;
+	ServerPort = Port;
+	// we listen on all available interfaces.
+	Info.iface = NULL;
+	Info.protocols = &Protocols[0];
+	// no extensions
+	Info.extensions = NULL;
+	Info.gid = -1;
+	Info.uid = -1;
+	Info.options = 0;
+	// tack on this object.
+	Info.user = this;
+
+	//Info.options |= LWS_SERVER_OPTION_DO_SSL_GLOBAL_INIT;
+	Info.options |= LWS_SERVER_OPTION_DISABLE_IPV6;
+
+	Context = lws_create_context(&Info);
+
+	if (Context == NULL)
+	{
+		ServerPort = 0;
+		delete Protocols;
+		Protocols = NULL;
+		return false; // couldn't create a server.
+	}
+
+	ConnectedCallBack = CallBack;
+#endif
+	return true;
+}
+
+bool FWebSocketServer::Tick()
+{
+#if USE_LIBWEBSOCKET
+	lws_service(Context, 0);
+	lws_callback_on_writable_all_protocol(Context, &Protocols[0]);
+#endif
+	return true;
+}
+
+FWebSocketServer::FWebSocketServer()
+{}
+
+FWebSocketServer::~FWebSocketServer()
+{
+#if USE_LIBWEBSOCKET
+	if (Context)
+	{
+		lws_context_destroy(Context);
+		Context = NULL;
+	}
+
+	 delete Protocols;
+	 Protocols = NULL;
+#endif
+}
+
+FString FWebSocketServer::Info()
+{
+#if USE_LIBWEBSOCKET
+	return FString::Printf(TEXT("%s:%i"), ANSI_TO_TCHAR(lws_canonical_hostname(Context)), ServerPort);
+#else // ! USE_LIBWEBSOCKET -- i.e. HTML5 currently does not allow this...
+	return FString(TEXT("NOT SUPPORTED"));
+#endif
+}
+
+// callback.
+#if USE_LIBWEBSOCKET
+static int unreal_networking_server
+	(
+		struct lws *Wsi,
+		enum lws_callback_reasons Reason,
+		void *User,
+		void *In,
+		size_t Len
+	)
+{
+	struct lws_context *Context = lws_get_context(Wsi);
+	PerSessionDataServer* BufferInfo = (PerSessionDataServer*)User;
+	FWebSocketServer* Server = (FWebSocketServer*)lws_context_user(Context);
+
+	switch (Reason)
+	{
+		case LWS_CALLBACK_ESTABLISHED:
+			{
+				BufferInfo->Socket = new FWebSocket(Context, Wsi);
+				Server->ConnectedCallBack.ExecuteIfBound(BufferInfo->Socket);
+				lws_set_timeout(Wsi, NO_PENDING_TIMEOUT, 0);
+			}
+			break;
+
+		case LWS_CALLBACK_RECEIVE:
+			if (BufferInfo->Socket->Context == Context) // UE-74107 -- bandaid until this file is removed in favor of using LwsWebSocketsManager.cpp & LwsWebSocket.cpp
+			{
+				BufferInfo->Socket->OnRawRecieve(In, Len);
+			}
+			lws_set_timeout(Wsi, NO_PENDING_TIMEOUT, 0);
+			break;
+
+		case LWS_CALLBACK_SERVER_WRITEABLE:
+			if (BufferInfo->Socket->Context == Context) // UE-68340 -- bandaid until this file is removed in favor of using LwsWebSocketsManager.cpp & LwsWebSocket.cpp
+			{
+				BufferInfo->Socket->OnRawWebSocketWritable(Wsi);
+			}
+			lws_set_timeout(Wsi, NO_PENDING_TIMEOUT, 0);
+			break;
+		case LWS_CALLBACK_CLIENT_CONNECTION_ERROR:
+			{
+				BufferInfo->Socket->ErrorCallBack.ExecuteIfBound();
+			}
+			break;
+		case LWS_CALLBACK_WSI_DESTROY:
+		case LWS_CALLBACK_PROTOCOL_DESTROY:
+		case LWS_CALLBACK_CLOSED:
+		case LWS_CALLBACK_CLOSED_HTTP:
+			break;
+	}
+
+	return 0;
+}
+#endif
diff --git a/Engine/Platforms/Plugins/WebSocketNetworking/Source/WebSocketNetworking/Private/WebSocketServer.h b/Engine/Platforms/Plugins/WebSocketNetworking/Source/WebSocketNetworking/Private/WebSocketServer.h
new file mode 100644
--- /dev/null
+++ b/Engine/Platforms/Plugins/WebSocketNetworking/Source/WebSocketNetworking/Private/WebSocketServer.h
@@ -0,0 +1,40 @@
+// Copyright 1998-2019 Epic Games, Inc. All Rights Reserved.
+//
+// Read http://lucumr.pocoo.org/2012/9/24/websockets-101/ for a nice intro to web sockets.
+// This uses https://libwebsockets.org/trac/libwebsockets
+#pragma  once
+#include "WebSocketNetworkingPrivate.h"
+
+class FWebSocketServer
+{
+public:
+
+	FWebSocketServer();
+	~FWebSocketServer();
+
+	/** Create a web socket server*/
+	bool Init(uint32 Port, FWebSocketClientConnectedCallBack);
+
+	/** Service libwebsocket */
+	bool Tick();
+
+	/** Describe this libwebsocket server */
+	FString Info();
+
+// this was made public because of cross-platform build issues
+public:
+
+	/** Callback for a new websocket connection to the server */
+	FWebSocketClientConnectedCallBack  ConnectedCallBack;
+
+	/** Internal libwebsocket context */
+	WebSocketInternalContext* Context;
+
+	/** Protocols serviced by this implementation */
+	WebSocketInternalProtocol* Protocols;
+
+	friend class FWebSocket;
+	uint32 ServerPort;
+};
+
+
diff --git a/Engine/Platforms/Plugins/WebSocketNetworking/Source/WebSocketNetworking/Private/WebsocketConnection.cpp b/Engine/Platforms/Plugins/WebSocketNetworking/Source/WebSocketNetworking/Private/WebsocketConnection.cpp
new file mode 100644
--- /dev/null
+++ b/Engine/Platforms/Plugins/WebSocketNetworking/Source/WebSocketNetworking/Private/WebsocketConnection.cpp
@@ -0,0 +1,227 @@
+// Copyright 1998-2019 Epic Games, Inc. All Rights Reserved.
+
+#include "WebSocketNetworkingPrivate.h"
+
+#include "IPAddress.h"
+#include "Sockets.h"
+#include "Net/NetworkProfiler.h"
+#include "Net/DataChannel.h"
+
+#include "WebSocketConnection.h"
+#include "WebSocketNetDriver.h"
+#include "WebSocket.h"
+#include "Runtime/PacketHandlers/PacketHandler/Public/PacketHandler.h"
+#include "PacketHandlers/StatelessConnectHandlerComponent.h"
+
+/*-----------------------------------------------------------------------------
+Declarations.
+-----------------------------------------------------------------------------*/
+
+// Size of a UDP header.
+#define IP_HEADER_SIZE     (20)
+#define UDP_HEADER_SIZE    (IP_HEADER_SIZE+8)
+#define WINSOCK_MAX_PACKET (512)
+
+UWebSocketConnection::UWebSocketConnection(const FObjectInitializer& ObjectInitializer) :
+Super(ObjectInitializer),
+WebSocket(NULL)
+{
+}
+
+void UWebSocketConnection::InitBase(UNetDriver* InDriver, class FSocket* InSocket, const FURL& InURL, EConnectionState InState, int32 InMaxPacket, int32 InPacketOverhead)
+{
+	// Pass the call up the chain
+	Super::InitBase(InDriver, InSocket, InURL, InState,
+		// Use the default packet size/overhead unless overridden by a child class
+		InMaxPacket == 0 ? WINSOCK_MAX_PACKET : InMaxPacket,
+		InPacketOverhead == 0 ? UDP_HEADER_SIZE : InPacketOverhead);
+
+}
+
+void UWebSocketConnection::InitLocalConnection(UNetDriver* InDriver, class FSocket* InSocket, const FURL& InURL, EConnectionState InState, int32 InMaxPacket, int32 InPacketOverhead)
+{
+	InitBase(InDriver, InSocket, InURL, InState,
+		// Use the default packet size/overhead unless overridden by a child class
+		InMaxPacket == 0 ? WINSOCK_MAX_PACKET : InMaxPacket,
+		InPacketOverhead == 0 ? UDP_HEADER_SIZE : InPacketOverhead);
+
+	// Figure out IP address from the host URL
+
+	// Initialize our send bunch
+	InitSendBuffer();
+}
+
+void UWebSocketConnection::InitRemoteConnection(UNetDriver* InDriver, class FSocket* InSocket, const FURL& InURL, const class FInternetAddr& InRemoteAddr, EConnectionState InState, int32 InMaxPacket, int32 InPacketOverhead)
+{
+	InitBase(InDriver, InSocket, InURL, InState,
+		// Use the default packet size/overhead unless overridden by a child class
+		InMaxPacket == 0 ? WINSOCK_MAX_PACKET : InMaxPacket,
+		InPacketOverhead == 0 ? UDP_HEADER_SIZE : InPacketOverhead);
+
+	// Initialize our send bunch
+	InitSendBuffer();
+
+	// This is for a client that needs to log in, setup ClientLoginState and ExpectedClientLoginMsgType to reflect that
+	SetClientLoginState(EClientLoginState::LoggingIn);
+	SetExpectedClientLoginMsgType(NMT_Hello);
+}
+
+void UWebSocketConnection::LowLevelSend(void* Data, int32 CountBits, FOutPacketTraits& Traits)
+{
+	const uint8* DataToSend = reinterpret_cast<uint8*>(Data);
+
+	// Process any packet modifiers
+	if (Handler.IsValid() && !Handler->GetRawSend())
+	{
+		const ProcessedPacket ProcessedData = Handler->Outgoing(reinterpret_cast<uint8*>(Data), CountBits, Traits);
+
+		if (!ProcessedData.bError)
+		{
+			DataToSend = ProcessedData.Data;
+			CountBits = ProcessedData.CountBits;
+		}
+		else
+		{
+			CountBits = 0;
+		}
+	}
+
+	uint32 CountBytes = FMath::DivideAndRoundUp(CountBits, 8);
+
+	if (CountBits > (MaxPacket * 8))
+	{
+		UE_LOG( LogNet, Warning, TEXT( "UWebSocketConnection::LowLevelSend: CountBytes > MaxPacketSize! Count: %i, MaxPacket: %i %s" ), CountBytes, MaxPacket, *Describe() );
+	}
+
+	bool bBlockSend = false;
+
+#if !UE_BUILD_SHIPPING
+	UE_LOG(LogNetTraffic, VeryVerbose, TEXT("UWebSocketConnection::LowLevelSend: Address: %s, CountBytes: %i"), *LowLevelGetRemoteAddress(true), CountBytes);
+
+	LowLevelSendDel.ExecuteIfBound((void*)DataToSend, CountBytes, bBlockSend);
+#endif
+
+	if (!bBlockSend && CountBytes > 0)
+	{
+		WebSocket->Send((uint8*)DataToSend, CountBytes);
+	}
+}
+
+FString UWebSocketConnection::LowLevelGetRemoteAddress(bool bAppendPort)
+{
+	return WebSocket->RemoteEndPoint(bAppendPort);
+}
+
+FString UWebSocketConnection::LowLevelDescribe()
+{
+	return FString::Printf
+		(
+		TEXT(" remote=%s local=%s state: %s"),
+		*WebSocket->RemoteEndPoint(true),
+		*WebSocket->LocalEndPoint(true),
+		State == USOCK_Pending ? TEXT("Pending")
+		: State == USOCK_Open ? TEXT("Open")
+		: State == USOCK_Closed ? TEXT("Closed")
+		: TEXT("Invalid")
+		);
+}
+
+void UWebSocketConnection::SetWebSocket(FWebSocket* InWebSocket)
+{
+	WebSocket = InWebSocket; 
+}
+
+FWebSocket* UWebSocketConnection::GetWebSocket()
+{
+	return WebSocket;
+}
+
+void UWebSocketConnection::Tick()
+{
+	UNetConnection::Tick();
+	WebSocket->Tick();
+}
+
+void UWebSocketConnection::FinishDestroy()
+{
+	Super::FinishDestroy();
+	if (WebSocket)
+	{
+#if USE_LIBWEBSOCKET
+		WebSocket->Context = NULL; // UE-68340
+#endif
+		delete WebSocket; 
+		WebSocket = NULL;
+	}
+}
+
+void UWebSocketConnection::ReceivedRawPacket(void* Data,int32 Count)
+{
+#if !UE_BUILD_SHIPPING
+	UE_LOG(LogNetTraffic, VeryVerbose, TEXT("UWebSocketConnection::ReceivedRawPacket: Address: %s, Count: %i"), *LowLevelGetRemoteAddress(true), Count);
+#endif
+
+	if (Count == 0 ||   // nothing to process
+		Driver == NULL) // connection closing
+	{
+		return;
+	}
+
+	uint8* DataRef = reinterpret_cast<uint8*>(Data);
+	if ( bChallengeHandshake )
+	{
+		// Process all incoming packets.
+		if (Driver->ConnectionlessHandler.IsValid() && Driver->StatelessConnectComponent.IsValid() && Driver->GetSocketSubsystem() != nullptr)
+		{
+			int32 Port;
+			TSharedPtr<FInternetAddr> IncomingAddress = Driver->GetSocketSubsystem()->CreateInternetAddr();
+			IncomingAddress->SetRawIp(WebSocket->GetRawRemoteAddr(Port));
+			IncomingAddress->SetPort(Port);
+
+			const ProcessedPacket UnProcessedPacket =
+									Driver->ConnectionlessHandler->IncomingConnectionless(IncomingAddress, DataRef, Count);
+	
+			TSharedPtr<StatelessConnectHandlerComponent> StatelessConnect = Driver->StatelessConnectComponent.Pin();
+
+			bool bRestartedHandshake = false;
+
+			if (!UnProcessedPacket.bError && StatelessConnect->HasPassedChallenge(IncomingAddress, bRestartedHandshake) &&
+					!bRestartedHandshake)
+			{
+				UE_LOG(LogNet, Log, TEXT("Server accepting post-challenge connection from: %s"), *(IncomingAddress->ToString(false)));
+				// Set the initial packet sequence from the handshake data
+				if (StatelessConnectComponent.IsValid())
+				{
+					int32 ServerSequence = 0;
+					int32 ClientSequence = 0;
+					StatelessConnect->GetChallengeSequence(ServerSequence, ClientSequence);
+					InitSequence(ClientSequence, ServerSequence);
+				}
+
+				if (Handler.IsValid())
+				{
+					Handler->BeginHandshaking();
+				}
+
+				bChallengeHandshake = false; // i.e. bPassedChallenge
+				UE_LOG(LogNet, Warning, TEXT("UWebSocketConnection::bChallengeHandshake: %s"), *LowLevelDescribe());
+				Count = FMath::DivideAndRoundUp(UnProcessedPacket.CountBits, 8);
+				if (Count > 0)
+				{
+					DataRef = UnProcessedPacket.Data;
+				}
+				else
+				{
+					return; // NO FURTHER DATA TO PROCESS
+				}
+			}
+			else
+			{
+				// WARNING: if here, it might be during (bInitialConnect) - which needs to be processed (ReceivedRawPacket)
+				//return;
+			}
+		}
+	}
+
+	UNetConnection::ReceivedRawPacket(DataRef,Count);
+}
diff --git a/Engine/Platforms/Plugins/WebSocketNetworking/Source/WebSocketNetworking/WebSocketNetworking.Build.cs b/Engine/Platforms/Plugins/WebSocketNetworking/Source/WebSocketNetworking/WebSocketNetworking.Build.cs
new file mode 100644
--- /dev/null
+++ b/Engine/Platforms/Plugins/WebSocketNetworking/Source/WebSocketNetworking/WebSocketNetworking.Build.cs
@@ -0,0 +1,34 @@
+// Copyright 1998-2019 Epic Games, Inc. All Rights Reserved.
+
+namespace UnrealBuildTool.Rules
+{
+	public class WebSocketNetworking  : ModuleRules
+	{
+		public WebSocketNetworking(ReadOnlyTargetRules Target) : base(Target)
+		{
+			PrivateDependencyModuleNames.AddRange(
+				new string[] {
+					"Core",
+					"CoreUObject",
+					"Engine",
+					"EngineSettings",
+					"ImageCore",
+					"Sockets",
+					"PacketHandler",
+					"OpenSSL",
+					"libWebSockets",
+					"zlib"
+				}
+			);
+
+			if (Target.Platform == UnrealTargetPlatform.HTML5)
+			{
+				PublicDefinitions.Add("USE_LIBWEBSOCKET=0");
+			}
+			else
+			{
+				PublicDefinitions.Add("USE_LIBWEBSOCKET=1");
+			}
+		}
+	}
+}
diff --git a/Engine/Platforms/Plugins/WebSocketNetworking/WebSocketNetworking.uplugin b/Engine/Platforms/Plugins/WebSocketNetworking/WebSocketNetworking.uplugin
new file mode 100644
--- /dev/null
+++ b/Engine/Platforms/Plugins/WebSocketNetworking/WebSocketNetworking.uplugin
@@ -0,0 +1,33 @@
+{
+	"FileVersion" : 3,
+	"Version" : 1,
+	"VersionName" : "1.0",
+	"FriendlyName" : "Experimental WebSocket Networking Plugin",
+	"Description" : "WebSocket Networking - NOTE: MUST disable all other existing NetDriverDefinitions in order to use WebSocketNetDriver. ALSO: MUST disable all PackHandlerComponents not supported by HTML5/Websockets (e.g. SteamAuthComponentModuleInterface)",
+	"Category" : "Misc",
+	"CreatedBy" : "Epic Games, Inc.",
+	"CreatedByURL" : "http://epicgames.com",
+	"DocsURL" : "",
+	"MarketplaceURL" : "",
+	"SupportURL" : "",
+	"EnabledByDefault" : false,
+	"CanContainContent" : false,
+	"IsBetaVersion" : false,
+	"Installed" : false,
+	"Modules" :
+	[
+		{
+			"Name" : "WebSocketNetworking",
+			"Type" : "Runtime",
+			"LoadingPhase" : "Default",
+			"WhitelistPlatforms" :
+			[
+				"Mac",
+				"Win32",
+				"Win64",
+				"Linux",
+				"HTML5"
+			]
+		}
+	]
+}
