Fix linker error
#rnx

diff --git a/Engine/Plugins/Runtime/WebSocketNetworking/Source/WebSocketNetworking/WebSocketNetworking.Build.cs b/Engine/Plugins/Runtime/WebSocketNetworking/Source/WebSocketNetworking/WebSocketNetworking.Build.cs
--- a/Engine/Plugins/Runtime/WebSocketNetworking/Source/WebSocketNetworking/WebSocketNetworking.Build.cs
+++ b/Engine/Plugins/Runtime/WebSocketNetworking/Source/WebSocketNetworking/WebSocketNetworking.Build.cs
@@ -10,6 +10,7 @@ namespace UnrealBuildTool.Rules
 				new string[] {
 					"Core",
 					"CoreUObject",
+					"NetCore",
 					"Engine",
 					"EngineSettings",
 					"ImageCore",
