Removing HTML5 target.

diff --git a/Engine/Plugins/Runtime/WebSocketNetworking/WebSocketNetworking.uplugin b/Engine/Plugins/Runtime/WebSocketNetworking/WebSocketNetworking.uplugin
--- a/Engine/Plugins/Runtime/WebSocketNetworking/WebSocketNetworking.uplugin
+++ b/Engine/Plugins/Runtime/WebSocketNetworking/WebSocketNetworking.uplugin
@@ -26,7 +26,6 @@
 				"Win32",
 				"Win64",
 				"Linux",
-				"HTML5"
 			]
 		}
 	]
