Fix mac remove control by reducing the size of the websocket receive buffer for mac

diff --git a/Engine/Plugins/Experimental/WebSocketNetworking/Source/WebSocketNetworking/Private/WebSocketServer.cpp b/Engine/Plugins/Experimental/WebSocketNetworking/Source/WebSocketNetworking/Private/WebSocketServer.cpp
--- a/Engine/Plugins/Experimental/WebSocketNetworking/Source/WebSocketNetworking/Private/WebSocketServer.cpp
+++ b/Engine/Plugins/Experimental/WebSocketNetworking/Source/WebSocketNetworking/Private/WebSocketServer.cpp
@@ -56,7 +56,12 @@ bool FWebSocketServer::Init(uint32 Port, FWebSocketClientConnectedCallBack CallB
 	Protocols[0].name = "binary";
 	Protocols[0].callback = unreal_networking_server;
 	Protocols[0].per_session_data_size = sizeof(PerSessionDataServer);
+
+#if PLATFORM_WINDOWS
 	Protocols[0].rx_buffer_size = 10 * 1024 * 1024;
+#else
+	Protocols[0].rx_buffer_size = 1 * 1024 * 1024;
+#endif
 
 	Protocols[1].name = nullptr;
 	Protocols[1].callback = nullptr;
