Exposed optional/opt-in webserver to WebSocket server using the already available functionality from the underlying libwebsockets implementation.

- This allows serving static content through the websocket server without taking an extra port for a webserver.
- Directories on disk can be mounted for serving using `IWebSocketServer::EnableWebServer(TArray<FWebSocketHttpMount> DirectoriesToServe)`.
- `IWebSocketServer::EnableWebServer()` should be called before `IWebSocketServer::Init`
- Fixed memory leak in FWebSocketServer where some manually created arrays were not using delete[].

diff --git a/Engine/Plugins/Experimental/WebSocketNetworking/Source/WebSocketNetworking/Private/WebSocketNetworkingPrivate.h b/Engine/Plugins/Experimental/WebSocketNetworking/Source/WebSocketNetworking/Private/WebSocketNetworkingPrivate.h
--- a/Engine/Plugins/Experimental/WebSocketNetworking/Source/WebSocketNetworking/Private/WebSocketNetworkingPrivate.h
+++ b/Engine/Plugins/Experimental/WebSocketNetworking/Source/WebSocketNetworking/Private/WebSocketNetworkingPrivate.h
@@ -12,5 +12,6 @@ class FWebSocketServer;
 typedef struct lws_context WebSocketInternalContext;
 typedef struct lws WebSocketInternal;
 typedef struct lws_protocols WebSocketInternalProtocol;
+typedef struct lws_http_mount WebSocketInternalHttpMount;
 
 DECLARE_LOG_CATEGORY_EXTERN(LogWebSocketNetworking, Warning, All);
diff --git a/Engine/Plugins/Experimental/WebSocketNetworking/Source/WebSocketNetworking/Private/WebSocketServer.cpp b/Engine/Plugins/Experimental/WebSocketNetworking/Source/WebSocketNetworking/Private/WebSocketServer.cpp
--- a/Engine/Plugins/Experimental/WebSocketNetworking/Source/WebSocketNetworking/Private/WebSocketServer.cpp
+++ b/Engine/Plugins/Experimental/WebSocketNetworking/Source/WebSocketNetworking/Private/WebSocketServer.cpp
@@ -22,7 +22,7 @@ enum class EFragmentationState : uint8 {
 struct PerSessionDataServer
 {
 	// Each session is actually a socket to a client
-	FWebSocket *Socket;
+	FWebSocket* Socket;
 	// Holds the concatenated message fragments.
 	TArray<uint8> FrameBuffer;
 	// The current state of the message being read.
@@ -39,8 +39,64 @@ static int unreal_networking_server(struct lws *wsi, enum lws_callback_reasons r
 	{
 		UE_LOG(LogWebSocketNetworking, Log, TEXT("websocket server: %s"), ANSI_TO_TCHAR(line));
 	}
+#endif // UE_BUILD_SHIPPING
+
+#endif // USE_LIBWEBSOCKET
+
+bool FWebSocketServer::IsHttpEnabled() const
+{
+	return bEnableHttp;
+}
+
+void FWebSocketServer::EnableHTTPServer(TArray<FWebSocketHttpMount> InDirectoriesToServe)
+{
+#if USE_LIBWEBSOCKET
+	bEnableHttp = true;
+
+	DirectoriesToServe = MoveTemp(InDirectoriesToServe);
+	const int NMounts = DirectoriesToServe.Num();
+
+	if(NMounts == 0)
+	{
+		return;
+	}
+
+	// Convert DirectoriesToServe to lws_http_mount for lws
+	LwsHttpMounts = new lws_http_mount[NMounts];
+
+	for(int i = 0; i < NMounts; i++) {
+		bool bLastMount = i == (NMounts - 1);
+		FWebSocketHttpMount& MountDir = DirectoriesToServe[i];
+		WebSocketInternalHttpMount* LWSMount = &LwsHttpMounts[i];
+		LWSMount->mount_next = bLastMount ? NULL : &LwsHttpMounts[i + 1];
+		LWSMount->mountpoint = MountDir.GetWebPath();
+		LWSMount->origin = MountDir.GetPathOnDisk();
+
+		if(!MountDir.HasDefaultFile())
+		{
+			LWSMount->def = MountDir.GetDefaultFile();
+		}
+		else
+		{
+			LWSMount->def = NULL;
+		}
+
+		LWSMount->protocol = NULL;
+		LWSMount->cgienv = NULL;
+		LWSMount->extra_mimetypes = NULL; // We may wish to expose this in future
+		LWSMount->interpret = NULL;
+		LWSMount->cgi_timeout = 0;
+		LWSMount->cache_max_age = 0;
+		LWSMount->auth_mask = 0;
+		LWSMount->cache_reusable = 0;
+		LWSMount->cache_revalidate = 0;
+		LWSMount->cache_intermediaries = 0;
+		LWSMount->origin_protocol = LWSMPRO_FILE;
+		LWSMount->mountpoint_len = FPlatformString::Strlen(MountDir.GetWebPath());
+		LWSMount->basic_auth_login_file = NULL;
+	}
 #endif
-#endif
+}
 
 bool FWebSocketServer::Init(uint32 Port, FWebSocketClientConnectedCallBack CallBack)
 {
@@ -86,12 +142,17 @@ bool FWebSocketServer::Init(uint32 Port, FWebSocketClientConnectedCallBack CallB
 	//Info.options |= LWS_SERVER_OPTION_DO_SSL_GLOBAL_INIT;
 	Info.options |= LWS_SERVER_OPTION_DISABLE_IPV6;
 
+	if(bEnableHttp && LwsHttpMounts != NULL)
+	{
+		Info.mounts = &LwsHttpMounts[0];
+	}
+
 	Context = lws_create_context(&Info);
 
 	if (Context == NULL)
 	{
 		ServerPort = 0;
-		delete Protocols;
+		delete[] Protocols;
 		Protocols = NULL;
 		return false; // couldn't create a server.
 	}
@@ -118,8 +179,11 @@ FWebSocketServer::~FWebSocketServer()
 		Context = NULL;
 	}
 
-	 delete Protocols;
+	 delete[] Protocols;
 	 Protocols = NULL;
+
+	 delete[] LwsHttpMounts;
+	 LwsHttpMounts = NULL;
 #endif
 }
 
@@ -143,7 +207,7 @@ static int unreal_networking_server
 		size_t Len
 	)
 {
-	struct lws_context *Context = lws_get_context(Wsi);
+	struct lws_context* Context = lws_get_context(Wsi);
 	PerSessionDataServer* BufferInfo = (PerSessionDataServer*)User;
 	FWebSocketServer* Server = (FWebSocketServer*)lws_context_user(Context);
 
@@ -201,8 +265,11 @@ static int unreal_networking_server
 			if (BufferInfo->Socket->Context == Context) // UE-68340 -- bandaid until this file is removed in favor of using LwsWebSocketsManager.cpp & LwsWebSocket.cpp
 			{
 				BufferInfo->Socket->OnRawWebSocketWritable(Wsi);
+
+				// Note: This particular lws callback reason gets hit in both ws and http cases as it used to signal that the server is in a writeable state.
+				// This means we only want to set an infinite timeout on genuine websocket connections, not http connections, otherwise they hang!
+				lws_set_timeout(Wsi, NO_PENDING_TIMEOUT, 0);
 			}
-			lws_set_timeout(Wsi, NO_PENDING_TIMEOUT, 0);
 			break;
 		case LWS_CALLBACK_CLIENT_CONNECTION_ERROR:
 			{
@@ -210,17 +277,22 @@ static int unreal_networking_server
 			}
 			break;
 		case LWS_CALLBACK_CLOSED:
-			if (BufferInfo->Socket)
+			if (BufferInfo->Socket->Context == Context)
 			{
 				BufferInfo->Socket->OnClose();
 			}
 			break;
 		case LWS_CALLBACK_WSI_DESTROY:
+			break;
 		case LWS_CALLBACK_PROTOCOL_DESTROY:
-		case LWS_CALLBACK_CLOSED_HTTP:
 			break;
 	}
 
+	// Check if http should be enabled or not, if so, use the in-built `lws_callback_http_dummy` which handles basic http requests
+	if(Server != nullptr && Server->IsHttpEnabled())
+	{
+		return lws_callback_http_dummy(Wsi, Reason, User, In, Len);
+	}
 	return 0;
 }
-#endif
+#endif
\ No newline at end of file
diff --git a/Engine/Plugins/Experimental/WebSocketNetworking/Source/WebSocketNetworking/Private/WebSocketServer.h b/Engine/Plugins/Experimental/WebSocketNetworking/Source/WebSocketNetworking/Private/WebSocketServer.h
--- a/Engine/Plugins/Experimental/WebSocketNetworking/Source/WebSocketNetworking/Private/WebSocketServer.h
+++ b/Engine/Plugins/Experimental/WebSocketNetworking/Source/WebSocketNetworking/Private/WebSocketServer.h
@@ -14,11 +14,14 @@ public:
 	
 	//~ Begin IWebSocketServer interface
 	virtual ~FWebSocketServer() override;
+	virtual void EnableHTTPServer(TArray<FWebSocketHttpMount> DirectoriesToServe) override;
 	virtual bool Init(uint32 Port, FWebSocketClientConnectedCallBack) override;
 	virtual void Tick() override;
 	virtual FString Info() override;
 	//~ End IWebSocketServer interface
 
+	bool IsHttpEnabled() const;
+
 // this was made public because of cross-platform build issues
 public:
 
@@ -33,6 +36,12 @@ public:
 
 	friend class FWebSocket;
 	uint32 ServerPort;
+
+private:
+	bool bEnableHttp = false;
+
+	TArray<FWebSocketHttpMount> DirectoriesToServe;
+	WebSocketInternalHttpMount* LwsHttpMounts = NULL;
 };
 
 
diff --git a/Engine/Plugins/Experimental/WebSocketNetworking/Source/WebSocketNetworking/Public/IWebSocketServer.h b/Engine/Plugins/Experimental/WebSocketNetworking/Source/WebSocketNetworking/Public/IWebSocketServer.h
--- a/Engine/Plugins/Experimental/WebSocketNetworking/Source/WebSocketNetworking/Public/IWebSocketServer.h
+++ b/Engine/Plugins/Experimental/WebSocketNetworking/Source/WebSocketNetworking/Public/IWebSocketServer.h
@@ -4,12 +4,67 @@
 
 #include "Containers/UnrealString.h"
 #include "WebSocketNetworkingDelegates.h"
+#include <string>
+
+/**
+ * Represents a directory on disk that can be mapped to a url route and served.
+ * 
+ * For example we may wish to map `C:/www/Public` to `/` of our webserver.
+ **/
+class WEBSOCKETNETWORKING_API FWebSocketHttpMount
+{
+public:
+
+	/**
+	 * Sets the absolute path on disk to the directory we wish to serve.
+	 * @param InPathOnDisk The absolute path to the directory.
+	 */
+	void SetPathOnDisk(const FString& InPathOnDisk) { PathOnDisk = std::string(TCHAR_TO_ANSI(*InPathOnDisk)); }
+	
+	/**
+	 * Sets The web url path to use for this mount, e.g. /images .
+	 * @param InWebPath The web url path we will map the directory to, e.g. /images.
+	 */
+	void SetWebPath(const FString& InWebPath) { WebPath = std::string(TCHAR_TO_ANSI(*InWebPath)); }
+
+	/**
+	 * Sets a file to serve when the root web path is requested.
+	 * @param InDefaultFile The file to serve when the root web path is request, e.g. index.html.
+	 */
+	void SetDefaultFile(const FString& InDefaultFile) { DefaultFile = std::string(TCHAR_TO_ANSI(*InDefaultFile)); }
+
+	const char* GetPathOnDisk() { return PathOnDisk.c_str(); }
+	const char* GetWebPath() { return WebPath.c_str(); }
+	const char* GetDefaultFile() { return DefaultFile.c_str(); }
+	bool HasDefaultFile() { return DefaultFile.empty(); }
+
+private:
+
+	// Note: The below members are std::string purposefully as lws requires char* strings and FString may not be char*.
+
+	/* The absolute path on disk to directory we wish to serve. */
+	std::string PathOnDisk;
+
+	/* The web url path to use for this mount, e.g. /images */
+	std::string WebPath;
+
+	/* When the root of the `WebPath` is requested, without a file, we can serve this file, e.g. index.html */
+	std::string DefaultFile = "index.html";
+};
 
 class WEBSOCKETNETWORKING_API IWebSocketServer
 {
 public:
 	virtual ~IWebSocketServer() {}
 
+	/**
+	 * Extends the websocket server to additionally serve static content from on-disk directories using HTTP.
+	 * These HTTP requests will use the same port as the one specified for websocket use in `Init`.
+	 * Note: This function will not start the webserver, call `Init` to start the websocket/webserver.
+	 * @param DirectoriesToServe The directories for the Webserver to serve.
+	 */
+	virtual void EnableHTTPServer(TArray<FWebSocketHttpMount> DirectoriesToServe) = 0;
+
 	/**
 	 * Initialize the server and start listening for messages.
 	 * @param Port the port to handle websockets messages on.
