Fix ASan heap use after free in WebSocketServer.cpp

diff --git a/Engine/Plugins/Experimental/WebSocketNetworking/Source/WebSocketNetworking/Private/WebSocketServer.cpp b/Engine/Plugins/Experimental/WebSocketNetworking/Source/WebSocketNetworking/Private/WebSocketServer.cpp
--- a/Engine/Plugins/Experimental/WebSocketNetworking/Source/WebSocketNetworking/Private/WebSocketServer.cpp
+++ b/Engine/Plugins/Experimental/WebSocketNetworking/Source/WebSocketNetworking/Private/WebSocketServer.cpp
@@ -154,6 +154,8 @@ bool FWebSocketServer::Init(uint32 Port, FWebSocketClientConnectedCallBack CallB
 		ServerPort = 0;
 		delete[] Protocols;
 		Protocols = NULL;
+		delete[] LwsHttpMounts;
+	 	LwsHttpMounts = NULL;
 		return false; // couldn't create a server.
 	}
 
@@ -173,10 +175,12 @@ void FWebSocketServer::Tick()
 FWebSocketServer::~FWebSocketServer()
 {
 #if USE_LIBWEBSOCKET
+
 	if (Context)
 	{
-		lws_context_destroy(Context);
+		lws_context* ExistingContext = Context;
 		Context = NULL;
+		lws_context_destroy(ExistingContext);
 	}
 
 	 delete[] Protocols;
@@ -277,9 +281,15 @@ static int unreal_networking_server
 			}
 			break;
 		case LWS_CALLBACK_CLOSED:
-			if (BufferInfo->Socket->Context == Context)
 			{
-				BufferInfo->Socket->OnClose();
+				if(Server != nullptr)
+				{
+					bool bShuttingDown = Server->Context == NULL;
+					if (!bShuttingDown && BufferInfo->Socket->Context == Context)
+					{
+						BufferInfo->Socket->OnClose();
+					}
+				}
 			}
 			break;
 		case LWS_CALLBACK_WSI_DESTROY:
