Fix using temp object pointers which lifetime go out of scope after the struct is created

diff --git a/Engine/Plugins/Experimental/WebSocketNetworking/Source/WebSocketNetworking/Private/WebSocket.cpp b/Engine/Plugins/Experimental/WebSocketNetworking/Source/WebSocketNetworking/Private/WebSocket.cpp
--- a/Engine/Plugins/Experimental/WebSocketNetworking/Source/WebSocketNetworking/Private/WebSocket.cpp
+++ b/Engine/Plugins/Experimental/WebSocketNetworking/Source/WebSocketNetworking/Private/WebSocket.cpp
@@ -76,8 +76,10 @@ FWebSocket::FWebSocket(
 
 	check(Context);
 
+	auto ServerAddressString = StringCast<ANSICHAR>(*ServerAddress.ToString(false));
+
 	struct lws_client_connect_info ConnectInfo = {
-			Context, TCHAR_TO_ANSI(*ServerAddress.ToString(false)), ServerAddress.GetPort(), false, "/", TCHAR_TO_ANSI(*ServerAddress.ToString(false)), TCHAR_TO_ANSI(*ServerAddress.ToString(false)), Protocols[1].name, -1, this
+			Context, ServerAddressString.Get(), ServerAddress.GetPort(), false, "/", ServerAddressString.Get(), ServerAddressString.Get(), Protocols[1].name, -1, this
 	};
 	Wsi = lws_client_connect_via_info(&ConnectInfo);
 	check(Wsi);
