Fix CIS warning in WebsocketNetworking

diff --git a/Engine/Plugins/Experimental/WebSocketNetworking/Source/WebSocketNetworking/Private/WebSocketServer.cpp b/Engine/Plugins/Experimental/WebSocketNetworking/Source/WebSocketNetworking/Private/WebSocketServer.cpp
--- a/Engine/Plugins/Experimental/WebSocketNetworking/Source/WebSocketNetworking/Private/WebSocketServer.cpp
+++ b/Engine/Plugins/Experimental/WebSocketNetworking/Source/WebSocketNetworking/Private/WebSocketServer.cpp
@@ -309,7 +309,6 @@ static int unreal_networking_server
 					const int32 OriginHeaderLength = lws_hdr_total_length(Wsi, WSI_TOKEN_ORIGIN);
 
 					FString OriginHeader;
-					FString ClientIP;
 
 					if (OriginHeaderLength != 0)
 					{
@@ -324,15 +323,13 @@ static int unreal_networking_server
 						}
 					}
 
-					char IPAddress[16];
-					int32 IPAddressLength = 0;
+					constexpr int32 IPAddressLength = 16;
+					char IPAddress[IPAddressLength];
+					
 					lws_get_peer_simple(Wsi, IPAddress, IPAddressLength);
 
-					if (IPAddressLength != 0)
-					{
-						ClientIP = UTF8_TO_TCHAR(IPAddress);
-					}
-				
+					const FString ClientIP =UTF8_TO_TCHAR(IPAddress);
+			
 					if (Server->FilterConnectionCallback.Execute(OriginHeader, ClientIP) == EWebsocketConnectionFilterResult::ConnectionRefused)
 					{
 						bRejectConnection = true;
