Remote Control expose websocket server bind address

diff --git a/Engine/Plugins/Experimental/WebSocketNetworking/Source/WebSocketNetworking/Private/WebSocketServer.cpp b/Engine/Plugins/Experimental/WebSocketNetworking/Source/WebSocketNetworking/Private/WebSocketServer.cpp
--- a/Engine/Plugins/Experimental/WebSocketNetworking/Source/WebSocketNetworking/Private/WebSocketServer.cpp
+++ b/Engine/Plugins/Experimental/WebSocketNetworking/Source/WebSocketNetworking/Private/WebSocketServer.cpp
@@ -1,6 +1,7 @@
 // Copyright Epic Games, Inc. All Rights Reserved.
 
 #include "WebSocketServer.h"
+
 #include "WebSocket.h"
 
 #if USE_LIBWEBSOCKET
@@ -23,7 +24,7 @@ struct PerSessionDataServer
 {
 	// Each session is actually a socket to a client
 	FWebSocket* Socket;
-	// Holds the concatenated message fragments.
+	// Holds the concatenated message fragments.er
 	TArray<uint8> FrameBuffer;
 	// The current state of the message being read.
 	EFragmentationState FragementationState = EFragmentationState::BeginFrame;
@@ -98,7 +99,7 @@ void FWebSocketServer::EnableHTTPServer(TArray<FWebSocketHttpMount> InDirectorie
 #endif
 }
 
-bool FWebSocketServer::Init(uint32 Port, FWebSocketClientConnectedCallBack CallBack)
+bool FWebSocketServer::Init(uint32 Port, FWebSocketClientConnectedCallBack CallBack, FString BindAddress)
 {
 #if USE_LIBWEBSOCKET
 #if !UE_BUILD_SHIPPING
@@ -128,14 +129,19 @@ bool FWebSocketServer::Init(uint32 Port, FWebSocketClientConnectedCallBack CallB
 	// look up libwebsockets.h for details.
 	Info.port = Port;
 	ServerPort = Port;
-	// we listen on all available interfaces.
+
 	Info.iface = NULL;
+	if (!BindAddress.IsEmpty())
+	{
+		Info.iface = StringCast<ANSICHAR>(*BindAddress).Get();
+	}
+
 	Info.protocols = &Protocols[0];
 	// no extensions
 	Info.extensions = NULL;
 	Info.gid = -1;
 	Info.uid = -1;
-	Info.options = 0;
+	Info.options = LWS_SERVER_OPTION_ALLOW_LISTEN_SHARE;
 	// tack on this object.
 	Info.user = this;
 
diff --git a/Engine/Plugins/Experimental/WebSocketNetworking/Source/WebSocketNetworking/Private/WebSocketServer.h b/Engine/Plugins/Experimental/WebSocketNetworking/Source/WebSocketNetworking/Private/WebSocketServer.h
--- a/Engine/Plugins/Experimental/WebSocketNetworking/Source/WebSocketNetworking/Private/WebSocketServer.h
+++ b/Engine/Plugins/Experimental/WebSocketNetworking/Source/WebSocketNetworking/Private/WebSocketServer.h
@@ -15,7 +15,7 @@ public:
 	//~ Begin IWebSocketServer interface
 	virtual ~FWebSocketServer() override;
 	virtual void EnableHTTPServer(TArray<FWebSocketHttpMount> DirectoriesToServe) override;
-	virtual bool Init(uint32 Port, FWebSocketClientConnectedCallBack) override;
+	virtual bool Init(uint32 Port, FWebSocketClientConnectedCallBack, FString BindAddress = TEXT("")) override;
 	virtual void SetFilterConnectionCallback(FWebSocketFilterConnectionCallback InFilterConnectionCallback) override;
 	virtual void Tick() override;
 	virtual FString Info() override;
diff --git a/Engine/Plugins/Experimental/WebSocketNetworking/Source/WebSocketNetworking/Public/IWebSocketServer.h b/Engine/Plugins/Experimental/WebSocketNetworking/Source/WebSocketNetworking/Public/IWebSocketServer.h
--- a/Engine/Plugins/Experimental/WebSocketNetworking/Source/WebSocketNetworking/Public/IWebSocketServer.h
+++ b/Engine/Plugins/Experimental/WebSocketNetworking/Source/WebSocketNetworking/Public/IWebSocketServer.h
@@ -69,9 +69,10 @@ public:
 	 * Initialize the server and start listening for messages.
 	 * @param Port the port to handle websockets messages on.
 	 * @param ClientConnectedCallback the handler called when a client has connected.
+	 * @param BindAddress Which address the server should bind to. If left empty, the server will bind to all interfaces.
 	 * @return whether the initialization was successful.
 	 */
-	virtual bool Init(uint32 Port, FWebSocketClientConnectedCallBack ClientConnectedCallback) = 0;
+	virtual bool Init(uint32 Port, FWebSocketClientConnectedCallBack ClientConnectedCallback, FString BindAddress = TEXT("")) = 0;
 
 	/**
 	 * Set a filter callback, used to reject connections based on the origin header.
diff --git a/Engine/Plugins/VirtualProduction/RemoteControl/Source/RemoteControlCommon/Public/RemoteControlSettings.h b/Engine/Plugins/VirtualProduction/RemoteControl/Source/RemoteControlCommon/Public/RemoteControlSettings.h
--- a/Engine/Plugins/VirtualProduction/RemoteControl/Source/RemoteControlCommon/Public/RemoteControlSettings.h
+++ b/Engine/Plugins/VirtualProduction/RemoteControl/Source/RemoteControlCommon/Public/RemoteControlSettings.h
@@ -287,6 +287,10 @@ public:
 	UPROPERTY(config, EditAnywhere, Category = "Remote Control Web Server", DisplayName = "Remote Control HTTP Server Port")
 	uint32 RemoteControlHttpServerPort = 30010;
 
+	/** The address to bind the websocket server to. 0.0.0.0 will open the connection to everyone on your network, while 127.0.0.1 will only allow local requests to come through.  */
+	UPROPERTY(config, EditAnywhere, Category = "Remote Control Web Server", DisplayName = "Remote Control Websocket Bind Address")
+	FString RemoteControlWebsocketServerBindAddress = TEXT("0.0.0.0");
+
 	/** The web remote control WebSocket server's port. */
 	UPROPERTY(config, EditAnywhere, Category = "Remote Control Web Server", DisplayName = "Remote Control WebSocket Server Port")
 	uint32 RemoteControlWebSocketServerPort = 30020;
diff --git a/Engine/Plugins/VirtualProduction/RemoteControl/Source/WebRemoteControl/Private/RemoteControlWebSocketServer.cpp b/Engine/Plugins/VirtualProduction/RemoteControl/Source/WebRemoteControl/Private/RemoteControlWebSocketServer.cpp
--- a/Engine/Plugins/VirtualProduction/RemoteControl/Source/WebRemoteControl/Private/RemoteControlWebSocketServer.cpp
+++ b/Engine/Plugins/VirtualProduction/RemoteControl/Source/WebRemoteControl/Private/RemoteControlWebSocketServer.cpp
@@ -119,17 +119,12 @@ void FWebsocketMessageRouter::UnbindRoute(const FString& MessageName)
 
 bool FRCWebSocketServer::Start(uint32 Port, TSharedPtr<FWebsocketMessageRouter> InRouter)
 {
-	if (!IsPortAvailable(Port))
-	{
-		return false;
-	}
-
 	FWebSocketClientConnectedCallBack CallBack;
 	CallBack.BindRaw(this, &FRCWebSocketServer::OnWebSocketClientConnected);
 
 	Server = FModuleManager::Get().LoadModuleChecked<IWebSocketNetworkingModule>(TEXT("WebSocketNetworking")).CreateServer();
 	
-	if (!Server || !Server->Init(Port, CallBack))
+	if (!Server || !Server->Init(Port, CallBack, GetDefault<URemoteControlSettings>()->RemoteControlWebsocketServerBindAddress))
 	{
 		Server.Reset();
 		return false;
diff --git a/Engine/Plugins/VirtualProduction/RemoteControl/Source/WebRemoteControl/Private/WebRemoteControl.cpp b/Engine/Plugins/VirtualProduction/RemoteControl/Source/WebRemoteControl/Private/WebRemoteControl.cpp
--- a/Engine/Plugins/VirtualProduction/RemoteControl/Source/WebRemoteControl/Private/WebRemoteControl.cpp
+++ b/Engine/Plugins/VirtualProduction/RemoteControl/Source/WebRemoteControl/Private/WebRemoteControl.cpp
@@ -2538,7 +2538,7 @@ void FWebRemoteControlModule::OnSettingsModified(UObject* Settings, FPropertyCha
 	const bool bIsWebServerStarted = HttpRouter.IsValid();
 	const bool bIsWebSocketServerStarted = WebSocketServer.IsRunning();
 	const bool bRestartHttpServer = RCSettings->RemoteControlHttpServerPort != HttpServerPort;
-	const bool bRestartWebSocketServer = RCSettings->RemoteControlWebSocketServerPort != WebSocketServerPort;
+	const bool bRestartWebSocketServer = RCSettings->RemoteControlWebSocketServerPort != WebSocketServerPort || RCSettings->RemoteControlWebsocketServerBindAddress != WebsocketServerBindAddress;
 
 	if (PropertyChangedEvent.ChangeType != EPropertyChangeType::Interactive)
 	{
@@ -2555,6 +2555,7 @@ void FWebRemoteControlModule::OnSettingsModified(UObject* Settings, FPropertyCha
 			|| (!bIsWebSocketServerStarted && RCSettings->bAutoStartWebSocketServer))
 		{
 			WebSocketServerPort = RCSettings->RemoteControlWebSocketServerPort;
+			WebsocketServerBindAddress = RCSettings->RemoteControlWebsocketServerBindAddress;
 			StopWebSocketServer();
 			StartWebSocketServer();
 		}
diff --git a/Engine/Plugins/VirtualProduction/RemoteControl/Source/WebRemoteControl/Private/WebRemoteControl.h b/Engine/Plugins/VirtualProduction/RemoteControl/Source/WebRemoteControl/Private/WebRemoteControl.h
--- a/Engine/Plugins/VirtualProduction/RemoteControl/Source/WebRemoteControl/Private/WebRemoteControl.h
+++ b/Engine/Plugins/VirtualProduction/RemoteControl/Source/WebRemoteControl/Private/WebRemoteControl.h
@@ -176,6 +176,9 @@ private:
 	/** Port of the remote control http server. */
 	uint32 HttpServerPort;
 
+	/** The address the websocket server is currently bound to. */
+	FString WebsocketServerBindAddress;
+
 	/** Port of the remote control websocket server. */
 	uint32 WebSocketServerPort;
 
