Fix server restart when changing RC settings
Fix Websocket keeping port open after changing it (Our version of lws has a bug that keeps listening on port even if lws_context_destroy is called

diff --git a/Engine/Plugins/Experimental/WebSocketNetworking/Source/WebSocketNetworking/Private/WebSocketServer.cpp b/Engine/Plugins/Experimental/WebSocketNetworking/Source/WebSocketNetworking/Private/WebSocketServer.cpp
--- a/Engine/Plugins/Experimental/WebSocketNetworking/Source/WebSocketNetworking/Private/WebSocketServer.cpp
+++ b/Engine/Plugins/Experimental/WebSocketNetworking/Source/WebSocketNetworking/Private/WebSocketServer.cpp
@@ -5,6 +5,15 @@
 #include "WebSocket.h"
 
 #if USE_LIBWEBSOCKET
+
+#if PLATFORM_WINDOWS
+#include "Windows/AllowWindowsPlatformTypes.h"
+#include "Windows/PreWindowsApi.h"
+#include <winsock2.h>
+#include "Windows/PostWindowsApi.h"
+#include "Windows/HideWindowsPlatformTypes.h"
+#endif
+
 // Work around a conflict between a UI namespace defined by engine code and a typedef in OpenSSL
 #define UI UI_ST
 THIRD_PARTY_INCLUDES_START
@@ -302,11 +311,22 @@ static int unreal_networking_server
 						BufferInfo->Socket->OnClose();
 					}
 				}
+				bRejectConnection = true;
 			}
 			break;
 		case LWS_CALLBACK_WSI_DESTROY:
 			break;
 		case LWS_CALLBACK_PROTOCOL_DESTROY:
+#if PLATFORM_WINDOWS
+		{
+			// There is a bug with our version of the LWS library that keeps a socket open even if we destroy the context so we have to manually shut it down.
+			lws_sockfd_type Fd = lws_get_socket_fd(Wsi);
+			if (Fd != INVALID_SOCKET)
+			{
+				closesocket(Fd);
+			}
+		}
+#endif
 			break;
 		case LWS_CALLBACK_FILTER_PROTOCOL_CONNECTION:
 			{
diff --git a/Engine/Plugins/VirtualProduction/RemoteControl/Source/WebRemoteControl/Private/WebRemoteControl.cpp b/Engine/Plugins/VirtualProduction/RemoteControl/Source/WebRemoteControl/Private/WebRemoteControl.cpp
--- a/Engine/Plugins/VirtualProduction/RemoteControl/Source/WebRemoteControl/Private/WebRemoteControl.cpp
+++ b/Engine/Plugins/VirtualProduction/RemoteControl/Source/WebRemoteControl/Private/WebRemoteControl.cpp
@@ -351,6 +351,7 @@ void FWebRemoteControlModule::StartupModule()
 
 	HttpServerPort = GetDefault<URemoteControlSettings>()->RemoteControlHttpServerPort;
 	WebSocketServerPort = GetDefault<URemoteControlSettings>()->RemoteControlWebSocketServerPort;
+	WebsocketServerBindAddress = GetDefault<URemoteControlSettings>()->RemoteControlWebsocketServerBindAddress;
 
 	WebSocketHandler = MakeUnique<FWebSocketMessageHandler>(&WebSocketServer, ActingClientId);
 
