QOL: Add helper methods for retreiving URL params from websocket connections

diff --git a/Engine/Plugins/Experimental/WebSocketNetworking/Source/WebSocketNetworking/Private/WebSocket.cpp b/Engine/Plugins/Experimental/WebSocketNetworking/Source/WebSocketNetworking/Private/WebSocket.cpp
--- a/Engine/Plugins/Experimental/WebSocketNetworking/Source/WebSocketNetworking/Private/WebSocket.cpp
+++ b/Engine/Plugins/Experimental/WebSocketNetworking/Source/WebSocketNetworking/Private/WebSocket.cpp
@@ -236,6 +236,37 @@ FString FWebSocket::LocalEndPoint(bool bAppendPort)
 #endif
 }
 
+TArray<FString> FWebSocket::GetUrlArgs() 
+{
+	TArray<FString> UrlArgs;
+#if USE_LIBWEBSOCKET
+	char Buffer[256];	
+	
+	int n = 0;
+	while ( lws_hdr_copy_fragment(Wsi, Buffer, sizeof(Buffer), WSI_TOKEN_HTTP_URI_ARGS, n) >= 0 ) 
+	{
+		UrlArgs.Add(FString(ANSI_TO_TCHAR(Buffer)));
+		n++;
+	}
+#endif
+
+	return UrlArgs;	
+}
+
+FString FWebSocket::GetUrlArgByName(const FString& Name) 
+{
+#if USE_LIBWEBSOCKET
+	char Buffer[256];	
+
+	if ( lws_get_urlarg_by_name(Wsi, TCHAR_TO_ANSI(*Name), Buffer, sizeof(Buffer)) != nullptr )
+	{
+		return FString(ANSI_TO_TCHAR(Buffer));
+	}
+#endif
+
+	return FString(TEXT(""));
+}
+
 void FWebSocket::Tick()
 {
 	HandlePacket();
diff --git a/Engine/Plugins/Experimental/WebSocketNetworking/Source/WebSocketNetworking/Private/WebSocket.h b/Engine/Plugins/Experimental/WebSocketNetworking/Source/WebSocketNetworking/Private/WebSocket.h
--- a/Engine/Plugins/Experimental/WebSocketNetworking/Source/WebSocketNetworking/Private/WebSocket.h
+++ b/Engine/Plugins/Experimental/WebSocketNetworking/Source/WebSocketNetworking/Private/WebSocket.h
@@ -37,6 +37,8 @@ public:
 	virtual FString RemoteEndPoint(bool bAppendPort) override;
 	virtual FString LocalEndPoint(bool bAppendPort) override;
 	virtual struct sockaddr_in* GetRemoteAddr() override { return &RemoteAddr; }
+	virtual TArray<FString> GetUrlArgs() override;
+	virtual FString GetUrlArgByName(const FString& Name) override;
 	//~ End INetworkingWebSocket interface
 
 // this was made public because of cross-platform build issues
diff --git a/Engine/Plugins/Experimental/WebSocketNetworking/Source/WebSocketNetworking/Public/INetworkingWebSocket.h b/Engine/Plugins/Experimental/WebSocketNetworking/Source/WebSocketNetworking/Public/INetworkingWebSocket.h
--- a/Engine/Plugins/Experimental/WebSocketNetworking/Source/WebSocketNetworking/Public/INetworkingWebSocket.h
+++ b/Engine/Plugins/Experimental/WebSocketNetworking/Source/WebSocketNetworking/Public/INetworkingWebSocket.h
@@ -31,4 +31,6 @@ public:
 	virtual FString RemoteEndPoint(bool bAppendPort) = 0;
 	virtual FString LocalEndPoint(bool bAppendPort) = 0;
 	virtual struct sockaddr_in* GetRemoteAddr() = 0;
+	virtual TArray<FString> GetUrlArgs() = 0;
+	virtual FString GetUrlArgByName(const FString& Name) = 0;
 };
\ No newline at end of file
