[PS2] Add: Ability to configure the embedded signalling server to serve content over HTTPS

diff --git a/Engine/Plugins/Experimental/WebSocketNetworking/Source/WebSocketNetworking/Private/WebSocketServer.cpp b/Engine/Plugins/Experimental/WebSocketNetworking/Source/WebSocketNetworking/Private/WebSocketServer.cpp
--- a/Engine/Plugins/Experimental/WebSocketNetworking/Source/WebSocketNetworking/Private/WebSocketServer.cpp
+++ b/Engine/Plugins/Experimental/WebSocketNetworking/Source/WebSocketNetworking/Private/WebSocketServer.cpp
@@ -6,22 +6,23 @@
 
 #if USE_LIBWEBSOCKET
 
-#if PLATFORM_WINDOWS
-#include "Windows/AllowWindowsPlatformTypes.h"
-#include <winsock2.h>
-#include "Windows/HideWindowsPlatformTypes.h"
-#endif
-
-// Work around a conflict between a UI namespace defined by engine code and a typedef in OpenSSL
-#define UI UI_ST
+	#if PLATFORM_WINDOWS
+		#include "Windows/AllowWindowsPlatformTypes.h"
+		#include <winsock2.h>
+		#include "Windows/HideWindowsPlatformTypes.h"
+	#endif
+
+	// Work around a conflict between a UI namespace defined by engine code and a typedef in OpenSSL
+	#define UI UI_ST
 THIRD_PARTY_INCLUDES_START
-#include "libwebsockets.h"
+	#include "libwebsockets.h"
 THIRD_PARTY_INCLUDES_END
-#undef UI
+	#undef UI
 #endif
 
 // The current state of the message being read.
-enum class EFragmentationState : uint8 {
+enum class EFragmentationState : uint8
+{
 	BeginFrame,
 	MessageFrame,
 };
@@ -37,17 +38,16 @@ struct PerSessionDataServer
 	EFragmentationState FragementationState = EFragmentationState::BeginFrame;
 };
 
-
 #if USE_LIBWEBSOCKET
 // real networking handler.
-static int unreal_networking_server(struct lws *wsi, enum lws_callback_reasons reason, void *user, void *in, size_t len);
+static int unreal_networking_server(struct lws* wsi, enum lws_callback_reasons reason, void* user, void* in, size_t len);
 
-#if !UE_BUILD_SHIPPING
-	inline void lws_debugLog(int level, const char *line)
-	{
-		UE_LOG(LogWebSocketNetworking, Log, TEXT("websocket server: %s"), ANSI_TO_TCHAR(line));
-	}
-#endif // UE_BUILD_SHIPPING
+	#if !UE_BUILD_SHIPPING
+inline void lws_debugLog(int level, const char* line)
+{
+	UE_LOG(LogWebSocketNetworking, Log, TEXT("websocket server: %s"), ANSI_TO_TCHAR(line));
+}
+	#endif // UE_BUILD_SHIPPING
 
 #endif // USE_LIBWEBSOCKET
 
@@ -56,15 +56,17 @@ bool FWebSocketServer::IsHttpEnabled() const
 	return bEnableHttp;
 }
 
-void FWebSocketServer::EnableHTTPServer(TArray<FWebSocketHttpMount> InDirectoriesToServe)
+void FWebSocketServer::EnableHTTPServer(TArray<FWebSocketHttpMount> InDirectoriesToServe, bool bInSecure, const FWebSocketServerCertificates& InCertificates)
 {
 #if USE_LIBWEBSOCKET
 	bEnableHttp = true;
+	bSecure = bInSecure;
+	Certificates = InCertificates;
 
 	DirectoriesToServe = MoveTemp(InDirectoriesToServe);
 	const int NMounts = DirectoriesToServe.Num();
 
-	if(NMounts == 0)
+	if (NMounts == 0)
 	{
 		return;
 	}
@@ -72,15 +74,16 @@ void FWebSocketServer::EnableHTTPServer(TArray<FWebSocketHttpMount> InDirectorie
 	// Convert DirectoriesToServe to lws_http_mount for lws
 	LwsHttpMounts = new lws_http_mount[NMounts];
 
-	for(int i = 0; i < NMounts; i++) {
-		bool bLastMount = i == (NMounts - 1);
-		FWebSocketHttpMount& MountDir = DirectoriesToServe[i];
+	for (int i = 0; i < NMounts; i++)
+	{
+		bool						bLastMount = i == (NMounts - 1);
+		FWebSocketHttpMount&		MountDir = DirectoriesToServe[i];
 		WebSocketInternalHttpMount* LWSMount = &LwsHttpMounts[i];
 		LWSMount->mount_next = bLastMount ? NULL : &LwsHttpMounts[i + 1];
 		LWSMount->mountpoint = MountDir.GetWebPath();
 		LWSMount->origin = MountDir.GetPathOnDisk();
 
-		if(!MountDir.HasDefaultFile())
+		if (!MountDir.HasDefaultFile())
 		{
 			LWSMount->def = MountDir.GetDefaultFile();
 		}
@@ -109,10 +112,10 @@ void FWebSocketServer::EnableHTTPServer(TArray<FWebSocketHttpMount> InDirectorie
 bool FWebSocketServer::Init(uint32 Port, FWebSocketClientConnectedCallBack CallBack, FString BindAddress)
 {
 #if USE_LIBWEBSOCKET
-#if !UE_BUILD_SHIPPING
+	#if !UE_BUILD_SHIPPING
 	// setup log level.
 	lws_set_log_level(LLL_ERR | LLL_WARN | LLL_NOTICE | LLL_DEBUG | LLL_INFO, lws_debugLog);
-#endif
+	#endif
 
 	Protocols = new lws_protocols[3];
 	FMemory::Memzero(Protocols, sizeof(lws_protocols) * 3);
@@ -121,11 +124,11 @@ bool FWebSocketServer::Init(uint32 Port, FWebSocketClientConnectedCallBack CallB
 	Protocols[0].callback = unreal_networking_server;
 	Protocols[0].per_session_data_size = sizeof(PerSessionDataServer);
 
-#if PLATFORM_WINDOWS
+	#if PLATFORM_WINDOWS
 	Protocols[0].rx_buffer_size = 10 * 1024 * 1024;
-#else
+	#else
 	Protocols[0].rx_buffer_size = 1 * 1024 * 1024;
-#endif
+	#endif
 
 	Protocols[1].name = nullptr;
 	Protocols[1].callback = nullptr;
@@ -152,10 +155,16 @@ bool FWebSocketServer::Init(uint32 Port, FWebSocketClientConnectedCallBack CallB
 	// tack on this object.
 	Info.user = this;
 
-	//Info.options |= LWS_SERVER_OPTION_DO_SSL_GLOBAL_INIT;
+	if (bSecure)
+	{
+		Info.options |= LWS_SERVER_OPTION_DO_SSL_GLOBAL_INIT;
+
+		Info.ssl_cert_filepath = Certificates.GetCertificateFilePath();
+		Info.ssl_private_key_filepath = Certificates.GetPrivateKeyFilePath();
+	}
 	Info.options |= LWS_SERVER_OPTION_DISABLE_IPV6;
 
-	if(bEnableHttp && LwsHttpMounts != NULL)
+	if (bEnableHttp && LwsHttpMounts != NULL)
 	{
 		Info.mounts = &LwsHttpMounts[0];
 	}
@@ -168,7 +177,7 @@ bool FWebSocketServer::Init(uint32 Port, FWebSocketClientConnectedCallBack CallB
 		delete[] Protocols;
 		Protocols = NULL;
 		delete[] LwsHttpMounts;
-	 	LwsHttpMounts = NULL;
+		LwsHttpMounts = NULL;
 		return false; // couldn't create a server.
 	}
 
@@ -201,11 +210,11 @@ FWebSocketServer::~FWebSocketServer()
 		lws_context_destroy(ExistingContext);
 	}
 
-	 delete[] Protocols;
-	 Protocols = NULL;
+	delete[] Protocols;
+	Protocols = NULL;
 
-	 delete[] LwsHttpMounts;
-	 LwsHttpMounts = NULL;
+	delete[] LwsHttpMounts;
+	LwsHttpMounts = NULL;
 #endif
 }
 
@@ -220,30 +229,28 @@ FString FWebSocketServer::Info()
 
 // callback.
 #if USE_LIBWEBSOCKET
-static int unreal_networking_server
-	(
-		struct lws *Wsi,
-		enum lws_callback_reasons Reason,
-		void *User,
-		void *In,
-		size_t Len
-	)
+static int unreal_networking_server(
+	struct lws*				  Wsi,
+	enum lws_callback_reasons Reason,
+	void*					  User,
+	void*					  In,
+	size_t					  Len)
 {
-	struct lws_context* Context = lws_get_context(Wsi);
+	struct lws_context*	  Context = lws_get_context(Wsi);
 	PerSessionDataServer* BufferInfo = (PerSessionDataServer*)User;
-	FWebSocketServer* Server = (FWebSocketServer*)lws_context_user(Context);
-	bool bRejectConnection = false;
+	FWebSocketServer*	  Server = (FWebSocketServer*)lws_context_user(Context);
+	bool				  bRejectConnection = false;
 
 	switch (Reason)
 	{
 		case LWS_CALLBACK_ESTABLISHED:
-			{
-				BufferInfo->Socket = new FWebSocket(Context, Wsi);
-				BufferInfo->FragementationState = EFragmentationState::BeginFrame;
-				Server->ConnectedCallBack.ExecuteIfBound(BufferInfo->Socket);
-				lws_set_timeout(Wsi, NO_PENDING_TIMEOUT, 0);
-			}
-			break;
+		{
+			BufferInfo->Socket = new FWebSocket(Context, Wsi);
+			BufferInfo->FragementationState = EFragmentationState::BeginFrame;
+			Server->ConnectedCallBack.ExecuteIfBound(BufferInfo->Socket);
+			lws_set_timeout(Wsi, NO_PENDING_TIMEOUT, 0);
+		}
+		break;
 
 		case LWS_CALLBACK_RECEIVE:
 			if (BufferInfo->Socket->Context == Context) // UE-74107 -- bandaid until this file is removed in favor of using LwsWebSocketsManager.cpp & LwsWebSocket.cpp
@@ -295,27 +302,27 @@ static int unreal_networking_server
 			}
 			break;
 		case LWS_CALLBACK_CLIENT_CONNECTION_ERROR:
-			{
-				BufferInfo->Socket->ErrorCallBack.ExecuteIfBound();
-			}
-			break;
+		{
+			BufferInfo->Socket->ErrorCallBack.ExecuteIfBound();
+		}
+		break;
 		case LWS_CALLBACK_CLOSED:
+		{
+			if (Server != nullptr)
 			{
-				if(Server != nullptr)
+				bool bShuttingDown = Server->Context == NULL;
+				if (!bShuttingDown && BufferInfo->Socket->Context == Context)
 				{
-					bool bShuttingDown = Server->Context == NULL;
-					if (!bShuttingDown && BufferInfo->Socket->Context == Context)
-					{
-						BufferInfo->Socket->OnClose();
-					}
+					BufferInfo->Socket->OnClose();
 				}
-				bRejectConnection = true;
 			}
-			break;
+			bRejectConnection = true;
+		}
+		break;
 		case LWS_CALLBACK_WSI_DESTROY:
 			break;
 		case LWS_CALLBACK_PROTOCOL_DESTROY:
-#if PLATFORM_WINDOWS
+	#if PLATFORM_WINDOWS
 		{
 			// There is a bug with our version of the LWS library that keeps a socket open even if we destroy the context so we have to manually shut it down.
 			lws_sockfd_type Fd = lws_get_socket_fd(Wsi);
@@ -324,48 +331,48 @@ static int unreal_networking_server
 				closesocket(Fd);
 			}
 		}
-#endif
-			break;
+	#endif
+		break;
 		case LWS_CALLBACK_FILTER_PROTOCOL_CONNECTION:
+		{
+			if (Server && Server->FilterConnectionCallback.IsBound())
 			{
-				if (Server && Server->FilterConnectionCallback.IsBound())
-            	{
-					const int32 OriginHeaderLength = lws_hdr_total_length(Wsi, WSI_TOKEN_ORIGIN);
+				const int32 OriginHeaderLength = lws_hdr_total_length(Wsi, WSI_TOKEN_ORIGIN);
 
-					FString OriginHeader;
+				FString OriginHeader;
 
-					if (OriginHeaderLength != 0)
+				if (OriginHeaderLength != 0)
+				{
+					constexpr uint32 MaximumHeaderLength = 1024;
+					if (OriginHeaderLength < MaximumHeaderLength)
 					{
-						constexpr uint32 MaximumHeaderLength = 1024;
-						if (OriginHeaderLength < MaximumHeaderLength)
-						{
-							char Origin[MaximumHeaderLength + 1];
-							memset(Origin, 0, MaximumHeaderLength + 1);
+						char Origin[MaximumHeaderLength + 1];
+						memset(Origin, 0, MaximumHeaderLength + 1);
 
-							lws_hdr_copy(Wsi, Origin, OriginHeaderLength + 1, WSI_TOKEN_ORIGIN);
-							OriginHeader = UTF8_TO_TCHAR(Origin);
-						}
+						lws_hdr_copy(Wsi, Origin, OriginHeaderLength + 1, WSI_TOKEN_ORIGIN);
+						OriginHeader = UTF8_TO_TCHAR(Origin);
 					}
+				}
 
-					constexpr int32 IPAddressLength = 16;
-					char IPAddress[IPAddressLength];
-					
-					lws_get_peer_simple(Wsi, IPAddress, IPAddressLength);
+				constexpr int32 IPAddressLength = 16;
+				char			IPAddress[IPAddressLength];
 
-					const FString ClientIP =UTF8_TO_TCHAR(IPAddress);
-			
-					if (Server->FilterConnectionCallback.Execute(OriginHeader, ClientIP) == EWebsocketConnectionFilterResult::ConnectionRefused)
-					{
-						bRejectConnection = true;
-					}
-				}
+				lws_get_peer_simple(Wsi, IPAddress, IPAddressLength);
+
+				const FString ClientIP = UTF8_TO_TCHAR(IPAddress);
 
-				break;
+				if (Server->FilterConnectionCallback.Execute(OriginHeader, ClientIP) == EWebsocketConnectionFilterResult::ConnectionRefused)
+				{
+					bRejectConnection = true;
+				}
 			}
+
+			break;
+		}
 	}
 
 	// Check if http should be enabled or not, if so, use the in-built `lws_callback_http_dummy` which handles basic http requests
-	if(Server != nullptr && Server->IsHttpEnabled())
+	if (Server != nullptr && Server->IsHttpEnabled())
 	{
 		return lws_callback_http_dummy(Wsi, Reason, User, In, Len);
 	}
diff --git a/Engine/Plugins/Experimental/WebSocketNetworking/Source/WebSocketNetworking/Private/WebSocketServer.h b/Engine/Plugins/Experimental/WebSocketNetworking/Source/WebSocketNetworking/Private/WebSocketServer.h
--- a/Engine/Plugins/Experimental/WebSocketNetworking/Source/WebSocketNetworking/Private/WebSocketServer.h
+++ b/Engine/Plugins/Experimental/WebSocketNetworking/Source/WebSocketNetworking/Private/WebSocketServer.h
@@ -2,35 +2,33 @@
 //
 // Read http://lucumr.pocoo.org/2012/9/24/websockets-101/ for a nice intro to web sockets.
 // This uses https://libwebsockets.org/trac/libwebsockets
-#pragma  once
+#pragma once
 #include "IWebSocketServer.h"
 #include "WebSocketNetworkingPrivate.h"
 
 class FWebSocketServer : public IWebSocketServer
 {
 public:
-
 	FWebSocketServer() = default;
-	
+
 	//~ Begin IWebSocketServer interface
 	virtual ~FWebSocketServer() override;
-	virtual void EnableHTTPServer(TArray<FWebSocketHttpMount> DirectoriesToServe) override;
-	virtual bool Init(uint32 Port, FWebSocketClientConnectedCallBack, FString BindAddress = TEXT("")) override;
-	virtual void SetFilterConnectionCallback(FWebSocketFilterConnectionCallback InFilterConnectionCallback) override;
-	virtual void Tick() override;
+	virtual void	EnableHTTPServer(TArray<FWebSocketHttpMount> DirectoriesToServe, bool bSecure, const FWebSocketServerCertificates& Certificates) override;
+	virtual bool	Init(uint32 Port, FWebSocketClientConnectedCallBack, FString BindAddress = TEXT("")) override;
+	virtual void	SetFilterConnectionCallback(FWebSocketFilterConnectionCallback InFilterConnectionCallback) override;
+	virtual void	Tick() override;
 	virtual FString Info() override;
 	//~ End IWebSocketServer interface
 
 	bool IsHttpEnabled() const;
 
-// this was made public because of cross-platform build issues
+	// this was made public because of cross-platform build issues
 public:
-
 	/** Callback for a new websocket connection to the server */
-	FWebSocketClientConnectedCallBack  ConnectedCallBack;
-	
+	FWebSocketClientConnectedCallBack ConnectedCallBack;
+
 	/** Handler called to filter incoming websocket connections. */
-    FWebSocketFilterConnectionCallback FilterConnectionCallback;
+	FWebSocketFilterConnectionCallback FilterConnectionCallback;
 
 	/** Internal libwebsocket context */
 	WebSocketInternalContext* Context;
@@ -44,8 +42,10 @@ public:
 private:
 	bool bEnableHttp = false;
 
+	bool bSecure = false;
+
+	FWebSocketServerCertificates Certificates;
+
 	TArray<FWebSocketHttpMount> DirectoriesToServe;
 	WebSocketInternalHttpMount* LwsHttpMounts = NULL;
 };
-
-
diff --git a/Engine/Plugins/Experimental/WebSocketNetworking/Source/WebSocketNetworking/Public/IWebSocketServer.h b/Engine/Plugins/Experimental/WebSocketNetworking/Source/WebSocketNetworking/Public/IWebSocketServer.h
--- a/Engine/Plugins/Experimental/WebSocketNetworking/Source/WebSocketNetworking/Public/IWebSocketServer.h
+++ b/Engine/Plugins/Experimental/WebSocketNetworking/Source/WebSocketNetworking/Public/IWebSocketServer.h
@@ -8,19 +8,18 @@
 
 /**
  * Represents a directory on disk that can be mapped to a url route and served.
- * 
+ *
  * For example we may wish to map `C:/www/Public` to `/` of our webserver.
  **/
 class WEBSOCKETNETWORKING_API FWebSocketHttpMount
 {
 public:
-
 	/**
 	 * Sets the absolute path on disk to the directory we wish to serve.
 	 * @param InPathOnDisk The absolute path to the directory.
 	 */
 	void SetPathOnDisk(const FString& InPathOnDisk) { PathOnDisk = std::string(TCHAR_TO_ANSI(*InPathOnDisk)); }
-	
+
 	/**
 	 * Sets The web url path to use for this mount, e.g. /images .
 	 * @param InWebPath The web url path we will map the directory to, e.g. /images.
@@ -36,10 +35,9 @@ public:
 	const char* GetPathOnDisk() { return PathOnDisk.c_str(); }
 	const char* GetWebPath() { return WebPath.c_str(); }
 	const char* GetDefaultFile() { return DefaultFile.c_str(); }
-	bool HasDefaultFile() { return DefaultFile.empty(); }
+	bool		HasDefaultFile() { return DefaultFile.empty(); }
 
 private:
-
 	// Note: The below members are std::string purposefully as lws requires char* strings and FString may not be char*.
 
 	/* The absolute path on disk to directory we wish to serve. */
@@ -52,18 +50,48 @@ private:
 	std::string DefaultFile = "index.html";
 };
 
+class WEBSOCKETNETWORKING_API FWebSocketServerCertificates
+{
+public:
+	/**
+	 * Sets the absolute path on disk to the file containing the SSL certificate.
+	 * @param InCertificateFilePath The absolute path to the file.
+	 */
+	void SetCertificateFilePath(const FString& InCertificateFilePath) { CertificateFilePath = std::string(TCHAR_TO_ANSI(*InCertificateFilePath)); }
+
+	/**
+	 * Sets the absolute path on disk to the file containing the SSL private key.
+	 * @param InPrivateKeyFilePath The absolute path to the file.
+	 */
+	void SetPrivateKeyFilePath(const FString& InPrivateKeyFilePath) { PrivateKeyFilePath = std::string(TCHAR_TO_ANSI(*InPrivateKeyFilePath)); }
+
+	const char* GetCertificateFilePath() { return CertificateFilePath.c_str(); }
+	const char* GetPrivateKeyFilePath() { return PrivateKeyFilePath.c_str(); }
+
+private:
+	// Note: The below members are std::string purposefully as lws requires char* strings and FString may not be char*.
+
+	// The absolute path to the certificate file
+	std::string CertificateFilePath;
+
+	// The absolute path to the private key file
+	std::string PrivateKeyFilePath;
+};
+
 class WEBSOCKETNETWORKING_API IWebSocketServer
 {
 public:
 	virtual ~IWebSocketServer() {}
 
 	/**
-	 * Extends the websocket server to additionally serve static content from on-disk directories using HTTP.
-	 * These HTTP requests will use the same port as the one specified for websocket use in `Init`.
+	 * Extends the websocket server to additionally serve static content from on-disk directories using HTTP or HTTPS.
+	 * These HTTP or HTTPS requests will use the same port as the one specified for websocket use in `Init`.
 	 * Note: This function will not start the webserver, call `Init` to start the websocket/webserver.
 	 * @param DirectoriesToServe The directories for the Webserver to serve.
+	 * @param bSecure Should the HTTP server serve the directories over HTTPS
+	 * @param Certificates FWebSocketServerCertificates containing paths to SSL certificates required when serving over HTTPS
 	 */
-	virtual void EnableHTTPServer(TArray<FWebSocketHttpMount> DirectoriesToServe) = 0;
+	virtual void EnableHTTPServer(TArray<FWebSocketHttpMount> DirectoriesToServe, bool bSecure = false, const FWebSocketServerCertificates& Certificates = FWebSocketServerCertificates()) = 0;
 
 	/**
 	 * Initialize the server and start listening for messages.
diff --git a/Engine/Plugins/Media/PixelStreaming2/Source/PixelStreaming2Editor/Private/PixelStreaming2Commands.cpp b/Engine/Plugins/Media/PixelStreaming2/Source/PixelStreaming2Editor/Private/PixelStreaming2Commands.cpp
--- a/Engine/Plugins/Media/PixelStreaming2/Source/PixelStreaming2Editor/Private/PixelStreaming2Commands.cpp
+++ b/Engine/Plugins/Media/PixelStreaming2/Source/PixelStreaming2Editor/Private/PixelStreaming2Commands.cpp
@@ -8,6 +8,7 @@ namespace UE::EditorPixelStreaming2
 	void FPixelStreaming2Commands::RegisterCommands()
 	{
 		UI_COMMAND(ExternalSignalling, "Use Remote Signalling Server", "Check this option if you wish to use a remote Signalling Server", EUserInterfaceActionType::RadioButton, FInputChord());
+		UI_COMMAND(ServeHttps, "Serve over HTTPS", "Check this option to serve the frontend over HTTPS using the embedded Signalling Server", EUserInterfaceActionType::RadioButton, FInputChord());
 		UI_COMMAND(StartSignalling, "Launch Signalling Server", "Launch a Signalling Server that will listen for connections on the ports specified above", EUserInterfaceActionType::Button, FInputChord());
 		UI_COMMAND(StopSignalling, "Stop Signalling Server", "Stop Signalling Server", EUserInterfaceActionType::Button, FInputChord());
 		UI_COMMAND(StreamLevelEditor, "Stream Level Editor", "Stream the Level Editor viewport", EUserInterfaceActionType::Button, FInputChord());
diff --git a/Engine/Plugins/Media/PixelStreaming2/Source/PixelStreaming2Editor/Private/PixelStreaming2Commands.h b/Engine/Plugins/Media/PixelStreaming2/Source/PixelStreaming2Editor/Private/PixelStreaming2Commands.h
--- a/Engine/Plugins/Media/PixelStreaming2/Source/PixelStreaming2Editor/Private/PixelStreaming2Commands.h
+++ b/Engine/Plugins/Media/PixelStreaming2/Source/PixelStreaming2Editor/Private/PixelStreaming2Commands.h
@@ -22,6 +22,7 @@ namespace UE::EditorPixelStreaming2
 		virtual void RegisterCommands() override;
 
 		TSharedPtr<FUICommandInfo> ExternalSignalling;
+		TSharedPtr<FUICommandInfo> ServeHttps;
 		TSharedPtr<FUICommandInfo> VP8;
 		TSharedPtr<FUICommandInfo> VP9;
 		TSharedPtr<FUICommandInfo> H264;
diff --git a/Engine/Plugins/Media/PixelStreaming2/Source/PixelStreaming2Editor/Private/PixelStreaming2EditorModule.cpp b/Engine/Plugins/Media/PixelStreaming2/Source/PixelStreaming2Editor/Private/PixelStreaming2EditorModule.cpp
--- a/Engine/Plugins/Media/PixelStreaming2/Source/PixelStreaming2Editor/Private/PixelStreaming2EditorModule.cpp
+++ b/Engine/Plugins/Media/PixelStreaming2/Source/PixelStreaming2Editor/Private/PixelStreaming2EditorModule.cpp
@@ -331,7 +331,9 @@ namespace UE::EditorPixelStreaming2
 		LaunchArgs.bPollUntilReady = false;
 		LaunchArgs.ReconnectionTimeoutSeconds = 30.0f;
 		LaunchArgs.ReconnectionIntervalSeconds = 2.0f;
-		LaunchArgs.ProcessArgs = FString::Printf(TEXT("--HttpPort=%d --StreamerPort=%d"), ViewerPort, StreamerPort);
+
+		FString CertificateArgs = FString::Printf(TEXT("--CertificatePath=%s --PrivateKeyPath=%s"), *SSLCertificatePath, *SSLPrivateKeyPath);
+		LaunchArgs.ProcessArgs = FString::Printf(TEXT("--HttpPort=%d --StreamerPort=%d --ServeHttps=%s %s"), ViewerPort, StreamerPort, bServeHttps ? TEXT("true") : TEXT("false"), bServeHttps ? *CertificateArgs : TEXT(""));
 
 		SignallingServer->Launch(LaunchArgs);
 	}
@@ -369,6 +371,20 @@ namespace UE::EditorPixelStreaming2
 		ViewerPort = InViewerPort;
 	}
 
+	void FPixelStreaming2EditorModule::SetServeHttps(bool bInServeHttps)
+	{
+		bServeHttps = bInServeHttps;
+	}
+
+	void FPixelStreaming2EditorModule::SetSSLCertificatePath(const FString& Path)
+	{
+		SSLCertificatePath = Path;
+	}
+	void FPixelStreaming2EditorModule::SetSSLPrivateKeyPath(const FString& Path)
+	{
+		SSLPrivateKeyPath = Path;
+	}
+
 	bool FPixelStreaming2EditorModule::ParseResolution(const TCHAR* InResolution, uint32& OutX, uint32& OutY)
 	{
 		if (*InResolution)
diff --git a/Engine/Plugins/Media/PixelStreaming2/Source/PixelStreaming2Editor/Private/PixelStreaming2EditorModule.h b/Engine/Plugins/Media/PixelStreaming2/Source/PixelStreaming2Editor/Private/PixelStreaming2EditorModule.h
--- a/Engine/Plugins/Media/PixelStreaming2/Source/PixelStreaming2Editor/Private/PixelStreaming2EditorModule.h
+++ b/Engine/Plugins/Media/PixelStreaming2/Source/PixelStreaming2Editor/Private/PixelStreaming2EditorModule.h
@@ -31,6 +31,12 @@ namespace UE::EditorPixelStreaming2
 		virtual int32	GetStreamerPort() override { return StreamerPort; };
 		virtual void	SetViewerPort(int32 InViewerPort) override;
 		virtual int32	GetViewerPort() override { return ViewerPort; };
+		virtual void	SetServeHttps(bool bServeHttps) override;
+		virtual bool	GetServeHttps() override { return bServeHttps; }
+		virtual void	SetSSLCertificatePath(const FString& Path) override;
+		virtual FString GetSSLCertificatePath() override { return SSLCertificatePath; }
+		virtual void	SetSSLPrivateKeyPath(const FString& Path) override;
+		virtual FString GetSSLPrivateKeyPath() override { return SSLPrivateKeyPath; }
 
 	private:
 		void InitEditorStreaming(IPixelStreaming2RTCModule& Module);
@@ -59,6 +65,13 @@ namespace UE::EditorPixelStreaming2
 #else
 		int32 ViewerPort = 80;
 #endif
+		//
+		bool bServeHttps = false;
+		//
+		FString SSLCertificatePath;
+		//
+		FString SSLPrivateKeyPath;
+
 		// The streamer used by the PixelStreaming2Editor module
 		TSharedPtr<IPixelStreaming2Streamer> EditorStreamer;
 
diff --git a/Engine/Plugins/Media/PixelStreaming2/Source/PixelStreaming2Editor/Private/PixelStreaming2Toolbar.cpp b/Engine/Plugins/Media/PixelStreaming2/Source/PixelStreaming2Editor/Private/PixelStreaming2Toolbar.cpp
--- a/Engine/Plugins/Media/PixelStreaming2/Source/PixelStreaming2Editor/Private/PixelStreaming2Toolbar.cpp
+++ b/Engine/Plugins/Media/PixelStreaming2/Source/PixelStreaming2Editor/Private/PixelStreaming2Toolbar.cpp
@@ -33,10 +33,12 @@
 #include "SlateFwd.h"
 #include "Widgets/Input/SNumericEntryBox.h"
 #include "PixelStreaming2PluginSettings.h"
-
+#include "Logging.h"
 #include "UtilsCoder.h"
 #include "Video/Encoders/Configs/VideoEncoderConfigAV1.h"
 #include "Video/Encoders/Configs/VideoEncoderConfigH264.h"
+#include "IDesktopPlatform.h"
+#include "DesktopPlatformModule.h"
 
 #define LOCTEXT_NAMESPACE "PixelStreaming2Editor"
 
@@ -63,8 +65,79 @@ bool GetUseRemoteSignallingServer()
 	return UPixelStreaming2PluginSettings::CVarEditorUseRemoteSignallingServer.GetValueOnAnyThread();
 }
 
+void SetServeHttps(bool bServeHttps)
+{
+	IPixelStreaming2EditorModule::Get().SetServeHttps(bServeHttps);
+}
+
+bool GetServeHttps()
+{
+	return IPixelStreaming2EditorModule::Get().GetServeHttps();
+}
+
+void SetSSLCertificatePath(const FString& Path)
+{
+	IPixelStreaming2EditorModule::Get().SetSSLCertificatePath(Path);
+}
+
+FString GetSSLCertificatePath()
+{
+	return IPixelStreaming2EditorModule::Get().GetSSLCertificatePath();
+}
+
+void SetSSLPrivateKeyPath(const FString& Path)
+{
+	IPixelStreaming2EditorModule::Get().SetSSLPrivateKeyPath(Path);
+}
+
+FString GetSSLPrivateKeyPath()
+{
+	return IPixelStreaming2EditorModule::Get().GetSSLPrivateKeyPath();
+}
+
 namespace UE::EditorPixelStreaming2
 {
+	TSharedRef<SWidget> CreateTextboxWithFileSelector(const FString& Label, float MaxWidth, TFunction<FString()> GetTextLambda, TFunction<void(const FText&)> OnTextChangedLambda, TFunction<void(const FText&, const ETextCommit::Type)> OnTextCommittedLambda, TFunction<void()> OnClickedLambda)
+	{
+		TSharedRef<SWidget> Widget = SNew(SHorizontalBox)
+			+ SHorizontalBox::Slot()
+				  .AutoWidth()
+				  .VAlign(VAlign_Center)
+				  .Padding(FMargin(36.0f, 3.0f, 8.0f, 3.0f))
+					  [SNew(STextBlock)
+							  .Text(FText::FromString(Label))
+							  .ColorAndOpacity(FSlateColor(FLinearColor(1.0f, 1.0f, 1.0f)))]
+			+ SHorizontalBox::Slot()
+				  .FillWidth(1.0f)
+				  .VAlign(VAlign_Center)
+				  .MaxWidth(MaxWidth)
+					  [SNew(SEditableTextBox)
+							  .OverflowPolicy(ETextOverflowPolicy::MiddleEllipsis)
+							  .Text_Lambda([GetTextLambda]() {
+								  return FText::FromString(GetTextLambda());
+							  })
+							  .OnTextChanged_Lambda([OnTextChangedLambda](const FText& Text) {
+								  OnTextChangedLambda(Text);
+							  })
+							  .OnTextCommitted_Lambda([OnTextCommittedLambda](const FText& Text, const ETextCommit::Type CommitType) {
+								  OnTextCommittedLambda(Text, CommitType);
+							  })]
+			+ SHorizontalBox::Slot()
+				  .AutoWidth()
+				  .HAlign(HAlign_Right)
+					  [SNew(SButton)
+							  .OnClicked_Lambda([OnClickedLambda]() {
+								  OnClickedLambda();
+								  return FReply::Handled();
+							  })
+							  .ButtonStyle(FAppStyle::Get(), "SimpleButton")
+								  [SNew(SImage)
+										  .Image(FAppStyle::Get().GetBrush("Icons.BrowseContent"))
+										  .ColorAndOpacity(FSlateColor::UseForeground())]];
+
+		return Widget;
+	}
+
 	FPixelStreaming2Toolbar::FPixelStreaming2Toolbar()
 	{
 		FPixelStreaming2Commands::Register();
@@ -89,6 +162,23 @@ namespace UE::EditorPixelStreaming2
 				return GetUseRemoteSignallingServer();
 			}));
 
+		PluginCommands->MapAction(
+			FPixelStreaming2Commands::Get().ServeHttps,
+			FExecuteAction::CreateLambda([]() {
+				SetServeHttps(!GetServeHttps());
+			}),
+			FCanExecuteAction::CreateLambda([] {
+				TSharedPtr<UE::PixelStreaming2Servers::IServer> SignallingServer = IPixelStreaming2EditorModule::Get().GetSignallingServer();
+				if (!SignallingServer.IsValid() || !SignallingServer->HasLaunched())
+				{
+					return true;
+				}
+				return false;
+			}),
+			FIsActionChecked::CreateLambda([]() {
+				return GetServeHttps();
+			}));
+
 		PluginCommands->MapAction(
 			FPixelStreaming2Commands::Get().StreamLevelEditor,
 			FExecuteAction::CreateLambda([]() {
@@ -287,6 +377,30 @@ namespace UE::EditorPixelStreaming2
 		TSharedPtr<UE::PixelStreaming2Servers::IServer> SignallingServer = IPixelStreaming2EditorModule::Get().GetSignallingServer();
 		if (!SignallingServer.IsValid() || !SignallingServer->HasLaunched())
 		{
+			MenuBuilder.AddMenuEntry(FPixelStreaming2Commands::Get().ServeHttps);
+			if (GetServeHttps())
+			{
+				TSharedRef<SWidget> SSLCertificateInputBlock = CreateTextboxWithFileSelector(
+					TEXT("SSL Certificate: "),																			   // Label
+					150,																								   // MaxWidth
+					[]() { return GetSSLCertificatePath(); },															   // GetTextLambda
+					[](const FText& Text) { SetSSLCertificatePath(Text.ToString()); },									   // OnTextChangedLambda
+					[](const FText& Text, const ETextCommit::Type CommitType) { SetSSLCertificatePath(Text.ToString()); }, // OnTextCommittedLambda
+					[this]() { OnOpenFileBrowserClicked(EFileType::Certificate); }										   // OnClickedLambda
+				);
+				MenuBuilder.AddWidget(SSLCertificateInputBlock, FText(), true);
+
+				TSharedRef<SWidget> SSLPrivateKeyInputBlock = CreateTextboxWithFileSelector(
+					TEXT("SSL Private Key: "),																			  // Label
+					150,																								  // MaxWidth
+					[]() { return GetSSLPrivateKeyPath(); },															  // GetTextLambda
+					[](const FText& Text) { SetSSLPrivateKeyPath(Text.ToString()); },									  // OnTextChangedLambda
+					[](const FText& Text, const ETextCommit::Type CommitType) { SetSSLPrivateKeyPath(Text.ToString()); }, // OnTextCommittedLambda
+					[this]() { OnOpenFileBrowserClicked(EFileType::PrivateKey); }										  // OnClickedLambda
+				);
+				MenuBuilder.AddWidget(SSLPrivateKeyInputBlock, FText(), true);
+			}
+
 			TSharedRef<SWidget> StreamerPortInputBlock = SNew(SHorizontalBox)
 				+ SHorizontalBox::Slot()
 					  .AutoWidth()
@@ -399,7 +513,7 @@ namespace UE::EditorPixelStreaming2
 									  .Padding(FMargin(32.0f, 3.0f))
 										  [SNew(STextBlock)
 												  .ColorAndOpacity(FSlateColor::UseSubduedForeground())
-												  .Text(FText::FromString(FString::Printf(TEXT("127.0.0.1:%d"), IPixelStreaming2EditorModule::Get().GetViewerPort())))],
+												  .Text(FText::FromString(FString::Printf(TEXT("%s://127.0.0.1:%d"), GetServeHttps() ? TEXT("https") : TEXT("http"), IPixelStreaming2EditorModule::Get().GetViewerPort())))],
 				FText());
 
 			TArray<TSharedPtr<FInternetAddr>> AdapterAddresses;
@@ -411,7 +525,7 @@ namespace UE::EditorPixelStreaming2
 											  .Padding(FMargin(32.0f, 3.0f))
 												  [SNew(STextBlock)
 														  .ColorAndOpacity(FSlateColor::UseSubduedForeground())
-														  .Text(FText::FromString(FString::Printf(TEXT("%s:%d"), *AdapterAddress->ToString(false), IPixelStreaming2EditorModule::Get().GetViewerPort())))],
+														  .Text(FText::FromString(FString::Printf(TEXT("%s://%s:%d"), GetServeHttps() ? TEXT("https") : TEXT("http"), *AdapterAddress->ToString(false), IPixelStreaming2EditorModule::Get().GetViewerPort())))],
 						FText());
 				}
 			}
@@ -485,6 +599,51 @@ namespace UE::EditorPixelStreaming2
 		MenuBuilder.EndSection();
 	}
 
+	void FPixelStreaming2Toolbar::OnOpenFileBrowserClicked(EFileType FileType)
+	{
+		bool bFileSelected = false;
+
+		if (IDesktopPlatform* DesktopPlatform = FDesktopPlatformModule::Get())
+		{
+			const FString Title = LOCTEXT("PixelStreaming2_FileBrowser", "Choose file").ToString();
+
+			if (LastBrowsePath.IsEmpty())
+			{
+				LastBrowsePath = FPaths::ProjectPluginsDir();
+			}
+
+			FString Filter = TEXT("");
+
+			TArray<FString> OutFiles;
+			bFileSelected = DesktopPlatform->OpenFileDialog(
+				FSlateApplication::Get().FindBestParentWindowHandleForDialogs(nullptr),
+				Title,
+				LastBrowsePath,
+				LastBrowsePath,
+				Filter,
+				EFileDialogFlags::None,
+				OutFiles);
+
+			if (!OutFiles.IsEmpty())
+			{
+				FString OutFilename = OutFiles[0];
+				switch (FileType)
+				{
+					case EFileType::Certificate:
+						UE_LOGFMT(LogPixelStreaming2Editor, Display, "Selecting certificate: {0}", OutFilename);
+						SetSSLCertificatePath(OutFilename);
+						break;
+					case EFileType::PrivateKey:
+						UE_LOGFMT(LogPixelStreaming2Editor, Display, "Selecting private key: {0}", OutFilename);
+						SetSSLPrivateKeyPath(OutFilename);
+						break;
+					default:
+						checkNoEntry();
+				}
+			}
+		}
+	}
+
 	TSharedRef<SWidget> FPixelStreaming2Toolbar::GeneratePixelStreaming2MenuContent(TSharedPtr<FUICommandList> InCommandList)
 	{
 		FToolMenuContext MenuContext(InCommandList);
diff --git a/Engine/Plugins/Media/PixelStreaming2/Source/PixelStreaming2Editor/Private/PixelStreaming2Toolbar.h b/Engine/Plugins/Media/PixelStreaming2/Source/PixelStreaming2Editor/Private/PixelStreaming2Toolbar.h
--- a/Engine/Plugins/Media/PixelStreaming2/Source/PixelStreaming2Editor/Private/PixelStreaming2Toolbar.h
+++ b/Engine/Plugins/Media/PixelStreaming2/Source/PixelStreaming2Editor/Private/PixelStreaming2Toolbar.h
@@ -16,11 +16,11 @@ namespace UE::EditorPixelStreaming2
 	public:
 		FPixelStreaming2Toolbar();
 		virtual ~FPixelStreaming2Toolbar();
-		void StartStreaming();
-		void StopStreaming();
+		void					   StartStreaming();
+		void					   StopStreaming();
 		static TSharedRef<SWidget> GeneratePixelStreaming2MenuContent(TSharedPtr<FUICommandList> InCommandList);
-		static FText GetActiveViewportName();
-		static const FSlateBrush* GetActiveViewportIcon();
+		static FText			   GetActiveViewportName();
+		static const FSlateBrush*  GetActiveViewportIcon();
 
 	private:
 		void RegisterMenus();
@@ -31,6 +31,18 @@ namespace UE::EditorPixelStreaming2
 		void RegisterVCamControls(FMenuBuilder& MenuBuilder);
 		void RegisterCodecConfig(FMenuBuilder& MenuBuilder);
 
+		enum class EFileType : uint8
+		{
+			Certificate,
+			PrivateKey
+		};
+
+		void OnOpenFileBrowserClicked(EFileType FileType);
+
 		TSharedPtr<class FUICommandList> PluginCommands;
+
+		// Store the last opened path so users don't have to constantly re-navigate to a certs folder
+		// when choosing cert and key
+		FString LastBrowsePath;
 	};
 } // namespace UE::EditorPixelStreaming2
diff --git a/Engine/Plugins/Media/PixelStreaming2/Source/PixelStreaming2Editor/Public/IPixelStreaming2EditorModule.h b/Engine/Plugins/Media/PixelStreaming2/Source/PixelStreaming2Editor/Public/IPixelStreaming2EditorModule.h
--- a/Engine/Plugins/Media/PixelStreaming2/Source/PixelStreaming2Editor/Public/IPixelStreaming2EditorModule.h
+++ b/Engine/Plugins/Media/PixelStreaming2/Source/PixelStreaming2Editor/Public/IPixelStreaming2EditorModule.h
@@ -106,4 +106,46 @@ public:
 	 * @return int32 The inbuilt c++ signalling server's port for viewers connect to
 	 */
 	virtual int32 GetViewerPort() = 0;
+
+	/**
+	 * Set whether frontend content should be served of HTTPS for the inbuilt c++ signalling server
+	 *
+	 * @param bServeHttps
+	 */
+	virtual void SetServeHttps(bool bServeHttps) = 0;
+
+	/**
+	 * Get the bool indicating if frontend content is being served over HTTPS for the inbuilt c++ signalling server
+	 *
+	 * @return bool true if frontend content is being served over https, else false.
+	 */
+	virtual bool GetServeHttps() = 0;
+
+	/**
+	 * Set the path to file containing the SSL certificate. Required if frontend content is being served of HTTPS for the inbuilt c++ signalling server
+	 *
+	 * @param Path The path to file containing the SSL certificate.
+	 */
+	virtual void SetSSLCertificatePath(const FString& Path) = 0;
+
+	/**
+	 * Get the path to file containing the SSL certificate.
+	 *
+	 * @return FString The path to file containing the SSL certificate.
+	 */
+	virtual FString GetSSLCertificatePath() = 0;
+
+	/**
+	 * Set the path to file containing the SSL private key. Required if frontend content is being served of HTTPS for the inbuilt c++ signalling server
+	 *
+	 * @param Path The path to file containing the SSL private key.
+	 */
+	virtual void SetSSLPrivateKeyPath(const FString& Path) = 0;
+
+	/**
+	 * Get the path to file containing the SSL private key.
+	 *
+	 * @return FString The path to file containing the SSL private key.
+	 */
+	virtual FString GetSSLPrivateKeyPath() = 0;
 };
diff --git a/Engine/Plugins/Media/PixelStreaming2/Source/PixelStreaming2Servers/Private/ServerUtils.cpp b/Engine/Plugins/Media/PixelStreaming2/Source/PixelStreaming2Servers/Private/ServerUtils.cpp
--- a/Engine/Plugins/Media/PixelStreaming2/Source/PixelStreaming2Servers/Private/ServerUtils.cpp
+++ b/Engine/Plugins/Media/PixelStreaming2/Source/PixelStreaming2Servers/Private/ServerUtils.cpp
@@ -26,19 +26,19 @@ namespace UE::PixelStreaming2Servers::Utils
 			return TSharedPtr<FMonitoredProcess>();
 		}
 
-		if(bRunAsScript)
+		if (bRunAsScript)
 		{
 // Get the executable we will use to run the scripts (e.g. cmd.exe on Windows)
 #if PLATFORM_WINDOWS
 			Args = FString::Printf(TEXT("/c \"\"%s\" %s\""), *ExecutableAbsPath, *Args);
 			ExecutableAbsPath = TEXT("cmd.exe");
-			
+
 #elif PLATFORM_LINUX
 			Args = FString::Printf(TEXT(" -- \"%s\" %s --nosudo"), *ExecutableAbsPath, *Args);
 			ExecutableAbsPath = TEXT("/usr/bin/bash");
 #elif PLATFORM_MAC
-            Args = FString::Printf(TEXT(" -- \"%s\" %s --nosudo"), *ExecutableAbsPath, *Args);
-            ExecutableAbsPath = TEXT("/bin/bash");
+			Args = FString::Printf(TEXT(" -- \"%s\" %s --nosudo"), *ExecutableAbsPath, *Args);
+			ExecutableAbsPath = TEXT("/bin/bash");
 #else
 			UE_LOG(LogPixelStreaming2Servers, Error, TEXT("Unsupported platform for Pixel Streaming."));
 			return TSharedPtr<FMonitoredProcess>();
@@ -46,9 +46,9 @@ namespace UE::PixelStreaming2Servers::Utils
 		}
 
 		TSharedPtr<FMonitoredProcess> ChildProcess = MakeShared<FMonitoredProcess>(
-			ExecutableAbsPath, 
-			Args, 
-			true, 
+			ExecutableAbsPath,
+			Args,
+			true,
 #if PLATFORM_MAC
 			// Pipes cause UE to lockup when destroying on Mac
 			false
@@ -98,7 +98,7 @@ namespace UE::PixelStreaming2Servers::Utils
 	FString QueryOrSetProcessArgs(FLaunchArgs& LaunchArgs, FString ArgKey, FString FallbackArgValue)
 	{
 		FString OutValue;
-		bool bExtractedValue = ExtractValueFromArgs(LaunchArgs.ProcessArgs, ArgKey, FallbackArgValue, OutValue);
+		bool	bExtractedValue = ExtractValueFromArgs(LaunchArgs.ProcessArgs, ArgKey, FallbackArgValue, OutValue);
 
 		// No key was present so we will inject our own.
 		if (!bExtractedValue)
@@ -119,7 +119,7 @@ namespace UE::PixelStreaming2Servers::Utils
 
 		OutResourcesDir = FPaths::ConvertRelativePathToFull(OutResourcesDir);
 
-		if(FPaths::DirectoryExists(OutResourcesDir))
+		if (FPaths::DirectoryExists(OutResourcesDir))
 		{
 			return true;
 		}
@@ -130,14 +130,14 @@ namespace UE::PixelStreaming2Servers::Utils
 	{
 		bool bResourceDirExists = GetResourcesDir(OutWebServersAbsPath);
 
-		if(!bResourceDirExists)
+		if (!bResourceDirExists)
 		{
 			return false;
 		}
 
 		OutWebServersAbsPath = OutWebServersAbsPath / TEXT("WebServers");
 
-		if(FPaths::DirectoryExists(OutWebServersAbsPath))
+		if (FPaths::DirectoryExists(OutWebServersAbsPath))
 		{
 			return true;
 		}
@@ -148,7 +148,7 @@ namespace UE::PixelStreaming2Servers::Utils
 	{
 		bool bServersDirExists = GetWebServersDir(OutAbsPath);
 
-		if(!bServersDirExists)
+		if (!bServersDirExists)
 		{
 			return false;
 		}
@@ -165,7 +165,7 @@ namespace UE::PixelStreaming2Servers::Utils
 #endif
 
 		IPlatformFile& PlatformFile = FPlatformFileManager::Get().GetPlatformFile();
-		if(PlatformFile.FileExists(*OutAbsPath))
+		if (PlatformFile.FileExists(*OutAbsPath))
 		{
 			return true;
 		}
@@ -175,16 +175,16 @@ namespace UE::PixelStreaming2Servers::Utils
 	TSharedPtr<FMonitoredProcess> DownloadPixelStreaming2Servers(bool bSkipIfPresent)
 	{
 		FString OutScriptPath;
-		if(bSkipIfPresent && GetDownloadedServer(OutScriptPath, FString(TEXT("SignallingWebServer"))))
+		if (bSkipIfPresent && GetDownloadedServer(OutScriptPath, FString(TEXT("SignallingWebServer"))))
 		{
 			UE_LOG(LogPixelStreaming2Servers, Log, TEXT("Found pixel streaming servers, skipping download."));
 			// empty process
 			return TSharedPtr<FMonitoredProcess>();
 		}
-		
+
 		bool bHasWebServersDir = GetWebServersDir(OutScriptPath);
 
-		if(!bHasWebServersDir)
+		if (!bHasWebServersDir)
 		{
 			UE_LOG(LogPixelStreaming2Servers, Error, TEXT("Could not download ps servers, no PixelStreaming2/Resources/WebServers directory found."));
 		}
@@ -212,6 +212,8 @@ namespace UE::PixelStreaming2Servers::Utils
 		FString StreamerPort = Utils::QueryOrSetProcessArgs(InLaunchArgs, TEXT("--StreamerPort="), FString::FromInt(Module.NextPort()));
 		FString SFUPort = Utils::QueryOrSetProcessArgs(InLaunchArgs, TEXT("--SFUPort="), FString::FromInt(Module.NextPort()));
 		FString HttpPort = Utils::QueryOrSetProcessArgs(InLaunchArgs, TEXT("--HttpPort="), FString::FromInt(Module.NextPort()));
+		FString ServeHttps = Utils::QueryOrSetProcessArgs(InLaunchArgs, TEXT("--ServeHttps="), TEXT("false"));
+		bool	bServeHttps = ServeHttps == TEXT("true");
 
 		// Construct endpoint urls
 
@@ -219,7 +221,7 @@ namespace UE::PixelStreaming2Servers::Utils
 		FURL SignallingStreamerUrl;
 		SignallingStreamerUrl.Protocol = TEXT("ws");
 		SignallingStreamerUrl.Host = TEXT("127.0.0.1");
-		SignallingStreamerUrl.Port = FCString::Atoi(*StreamerPort); 
+		SignallingStreamerUrl.Port = FCString::Atoi(*StreamerPort);
 		SignallingStreamerUrl.Map = FString();
 
 		// Players
@@ -231,7 +233,7 @@ namespace UE::PixelStreaming2Servers::Utils
 
 		// Webserver
 		FURL WebServerUrl;
-		WebServerUrl.Protocol = TEXT("http");
+		WebServerUrl.Protocol = bServeHttps ? TEXT("https") : TEXT("http");
 		WebServerUrl.Host = TEXT("127.0.0.1");
 		WebServerUrl.Port = FCString::Atoi(*HttpPort);
 		WebServerUrl.Map = FString();
@@ -257,16 +259,16 @@ namespace UE::PixelStreaming2Servers::Utils
 	FString ToString(TArrayView<uint8> UTF8Bytes)
 	{
 		FUTF8ToTCHAR Converted((const ANSICHAR*)UTF8Bytes.GetData(), UTF8Bytes.Num());
-		FString OutString(Converted.Length(), Converted.Get());
+		FString		 OutString(Converted.Length(), Converted.Get());
 		return OutString;
 	}
 
 	FString ToString(TSharedPtr<FJsonObject> JSONObj)
 	{
 		FString Res;
-		auto JsonWriter = TJsonWriterFactory<TCHAR, TCondensedJsonPrintPolicy<TCHAR>>::Create(&Res);
-		bool bSerialized = FJsonSerializer::Serialize(JSONObj.ToSharedRef(), JsonWriter);
-		if(!bSerialized)
+		auto	JsonWriter = TJsonWriterFactory<TCHAR, TCondensedJsonPrintPolicy<TCHAR>>::Create(&Res);
+		bool	bSerialized = FJsonSerializer::Serialize(JSONObj.ToSharedRef(), JsonWriter);
+		if (!bSerialized)
 		{
 			UE_LOG(LogPixelStreaming2Servers, Error, TEXT("Failed to stringify JSON object."));
 		}
@@ -276,7 +278,7 @@ namespace UE::PixelStreaming2Servers::Utils
 	TSharedPtr<FJsonObject> ToJSON(const FString& InString)
 	{
 		TSharedPtr<FJsonObject> OutJSON = MakeShared<FJsonObject>();
-		const auto JsonReader = TJsonReaderFactory<TCHAR>::Create(InString);
+		const auto				JsonReader = TJsonReaderFactory<TCHAR>::Create(InString);
 		if (FJsonSerializer::Deserialize(JsonReader, OutJSON))
 		{
 			return OutJSON;
@@ -284,4 +286,4 @@ namespace UE::PixelStreaming2Servers::Utils
 		return nullptr;
 	}
 
-} // UE::PixelStreaming2Servers::Utils
\ No newline at end of file
+} // namespace UE::PixelStreaming2Servers::Utils
\ No newline at end of file
diff --git a/Engine/Plugins/Media/PixelStreaming2/Source/PixelStreaming2Servers/Private/SignallingServer.cpp b/Engine/Plugins/Media/PixelStreaming2/Source/PixelStreaming2Servers/Private/SignallingServer.cpp
--- a/Engine/Plugins/Media/PixelStreaming2/Source/PixelStreaming2Servers/Private/SignallingServer.cpp
+++ b/Engine/Plugins/Media/PixelStreaming2/Source/PixelStreaming2Servers/Private/SignallingServer.cpp
@@ -19,6 +19,7 @@ namespace UE::PixelStreaming2Servers
 		PlayerMessageHandlers.Add("subscribe").BindRaw(this, &FSignallingServer::OnPlayerSubscribeMessage);
 		PlayerMessageHandlers.Add("unsubscribe").BindRaw(this, &FSignallingServer::OnPlayerUnsubscribeMessage);
 		PlayerMessageHandlers.Add("stats").BindRaw(this, &FSignallingServer::OnPlayerStatsMessage);
+		PlayerMessageHandlers.Add("ping").BindRaw(this, &FSignallingServer::OnPlayerPingMessage);
 	}
 
 	void FSignallingServer::Stop()
@@ -88,11 +89,29 @@ namespace UE::PixelStreaming2Servers
 			return false;
 		}
 
+		FString						 ServeHttpsString = Utils::QueryOrSetProcessArgs(InLaunchArgs, TEXT("--ServeHttps="), TEXT("false"));
+		bool						 bServeHttps = ServeHttpsString == TEXT("true");
+		FWebSocketServerCertificates Certificates;
+		if (bServeHttps)
+		{
+			FString CertificatePath = Utils::QueryOrSetProcessArgs(InLaunchArgs, TEXT("--CertificatePath="), TEXT(""));
+			if (!CertificatePath.IsEmpty())
+			{
+				Certificates.SetCertificateFilePath(CertificatePath);
+			}
+
+			FString PrivateKeyPath = Utils::QueryOrSetProcessArgs(InLaunchArgs, TEXT("--PrivateKeyPath="), TEXT(""));
+			if (!PrivateKeyPath.IsEmpty())
+			{
+				Certificates.SetPrivateKeyFilePath(PrivateKeyPath);
+			}
+		}
+
 		/*
 		 * --------------- Players websocket server ---------------
 		 */
 		PlayersWS = MakeUnique<FWebSocketServerWrapper>();
-		PlayersWS->EnableWebServer(GenerateDirectoriesToServe());
+		PlayersWS->EnableWebServer(GenerateDirectoriesToServe(), bServeHttps, Certificates);
 		PlayersWS->OnMessage.AddRaw(this, &FSignallingServer::OnPlayerMessage);
 		PlayersWS->OnOpenConnection.AddRaw(this, &FSignallingServer::OnPlayerConnected);
 		PlayersWS->OnClosedConnection.AddRaw(this, &FSignallingServer::OnPlayerDisconnected);
@@ -126,7 +145,7 @@ namespace UE::PixelStreaming2Servers
 	TArray<FWebSocketHttpMount> FSignallingServer::GenerateDirectoriesToServe() const
 	{
 		FString ServersDir;
-		bool bServersDirExists = Utils::GetWebServersDir(ServersDir);
+		bool	bServersDirExists = Utils::GetWebServersDir(ServersDir);
 		if (bServersDirExists)
 		{
 			ServersDir = ServersDir / TEXT("SignallingWebServer");
@@ -141,7 +160,7 @@ namespace UE::PixelStreaming2Servers
 		if (!bServersDirExists)
 		{
 			FString OutResourcesDir;
-			bool bResourcesDirExists = Utils::GetResourcesDir(OutResourcesDir);
+			bool	bResourcesDirExists = Utils::GetResourcesDir(OutResourcesDir);
 			FString NotFoundDir = OutResourcesDir / TEXT("NotFound");
 
 			if (bResourcesDirExists && FPaths::DirectoryExists(NotFoundDir))
@@ -225,7 +244,7 @@ namespace UE::PixelStreaming2Servers
 
 		// We don't want to make the connections shared to prevent someone accidentally holding on to it. So we use it raw here
 		FWebSocketConnection* PlayerWS = (*PlayersWS->GetConnections().Find(PlayerConnectionId)).Get();
-		bool bUESendsOffer = !PlayerWS->GetUrlArgs().Contains(TEXT("OfferToReceive=true"));
+		bool				  bUESendsOffer = !PlayerWS->GetUrlArgs().Contains(TEXT("OfferToReceive=true"));
 
 		// Send "playerConnected" message to streamer which kicks off making a new RTC connection
 		TSharedRef<FJsonObject> OnPlayerConnectedJSON = MakeShared<FJsonObject>();
@@ -306,7 +325,7 @@ namespace UE::PixelStreaming2Servers
 		const FString Msg = Utils::ToString(Message);
 		UE_LOG(LogPixelStreaming2Servers, Log, TEXT("From Streamer id=%d: %s"), ConnectionId, *Msg);
 
-		FString MsgType;
+		FString					MsgType;
 		TSharedPtr<FJsonObject> JSONObj = ParseMessage(Msg, MsgType);
 		if (!JSONObj)
 		{
@@ -352,7 +371,7 @@ namespace UE::PixelStreaming2Servers
 		const FString Msg = Utils::ToString(Message);
 		UE_LOG(LogPixelStreaming2Servers, Log, TEXT("From Player id=%d: %s"), ConnectionId, *Msg);
 
-		FString MsgType;
+		FString					MsgType;
 		TSharedPtr<FJsonObject> JSONObj = ParseMessage(Msg, MsgType);
 		if (!JSONObj)
 		{
@@ -369,7 +388,7 @@ namespace UE::PixelStreaming2Servers
 			if (!PlayerSubscriptions.Contains(ConnectionId))
 			{
 				TArray<FString> StreamerConnections = StreamersWS->GetConnectionNames();
-				if(StreamerConnections.Num() == 0)
+				if (StreamerConnections.Num() == 0)
 				{
 					UE_LOG(LogPixelStreaming2Servers, Error, TEXT("Player %d sent a message, but no streamers were connected"), ConnectionId);
 					return;
@@ -397,7 +416,7 @@ namespace UE::PixelStreaming2Servers
 
 	void FSignallingServer::OnStreamerPingMessage(uint16 ConnectionId, TSharedPtr<FJsonObject> JSONObj)
 	{
-		const double UnixTime = FDateTime::UtcNow().ToUnixTimestamp();
+		const double			UnixTime = FDateTime::UtcNow().ToUnixTimestamp();
 		TSharedRef<FJsonObject> PongJSON = MakeShared<FJsonObject>();
 		PongJSON->SetStringField("type", "pong");
 		PongJSON->SetNumberField("time", UnixTime);
@@ -418,10 +437,19 @@ namespace UE::PixelStreaming2Servers
 		PlayersWS->Close(PlayerConnectionId);
 	}
 
+	void FSignallingServer::OnPlayerPingMessage(uint16 ConnectionId, TSharedPtr<FJsonObject> JSONObj)
+	{
+		const double			UnixTime = FDateTime::UtcNow().ToUnixTimestamp();
+		TSharedRef<FJsonObject> PongJSON = MakeShared<FJsonObject>();
+		PongJSON->SetStringField("type", "pong");
+		PongJSON->SetNumberField("time", UnixTime);
+		SendPlayerMessage(ConnectionId, PongJSON);
+	}
+
 	void FSignallingServer::OnPlayerListStreamersMessage(uint16 ConnectionId, TSharedPtr<FJsonObject> JSONObj)
 	{
-		TSharedRef<FJsonObject> listJSON = MakeShared<FJsonObject>();
-		const TArray<FString> Names = StreamersWS->GetConnectionNames();
+		TSharedRef<FJsonObject>		   listJSON = MakeShared<FJsonObject>();
+		const TArray<FString>		   Names = StreamersWS->GetConnectionNames();
 		TArray<TSharedPtr<FJsonValue>> JsonNames;
 		for (const FString& Name : Names)
 		{
@@ -457,7 +485,7 @@ namespace UE::PixelStreaming2Servers
 
 	void FSignallingServer::GetNumStreamers(TFunction<void(uint16)> OnNumStreamersReceived)
 	{
-		if(StreamersWS)
+		if (StreamersWS)
 		{
 			OnNumStreamersReceived(StreamersWS->Count());
 		}
diff --git a/Engine/Plugins/Media/PixelStreaming2/Source/PixelStreaming2Servers/Private/SignallingServer.h b/Engine/Plugins/Media/PixelStreaming2/Source/PixelStreaming2Servers/Private/SignallingServer.h
--- a/Engine/Plugins/Media/PixelStreaming2/Source/PixelStreaming2Servers/Private/SignallingServer.h
+++ b/Engine/Plugins/Media/PixelStreaming2/Source/PixelStreaming2Servers/Private/SignallingServer.h
@@ -54,6 +54,7 @@ namespace UE::PixelStreaming2Servers
 		void OnStreamerIdMessage(uint16 ConnectionId, TSharedPtr<FJsonObject> JSONObj);
 		void OnStreamerPingMessage(uint16 ConnectionId, TSharedPtr<FJsonObject> JSONObj);
 		void OnStreamerDisconnectMessage(uint16 ConnectionId, TSharedPtr<FJsonObject> JSONObj);
+		void OnPlayerPingMessage(uint16 ConnectionId, TSharedPtr<FJsonObject> JSONObj);
 		void OnPlayerListStreamersMessage(uint16 ConnectionId, TSharedPtr<FJsonObject> JSONObj);
 		void OnPlayerSubscribeMessage(uint16 ConnectionId, TSharedPtr<FJsonObject> JSONObj);
 		void OnPlayerUnsubscribeMessage(uint16 ConnectionId, TSharedPtr<FJsonObject> JSONObj);
diff --git a/Engine/Plugins/Media/PixelStreaming2/Source/PixelStreaming2Servers/Private/WebSocketServerWrapper.cpp b/Engine/Plugins/Media/PixelStreaming2/Source/PixelStreaming2Servers/Private/WebSocketServerWrapper.cpp
--- a/Engine/Plugins/Media/PixelStreaming2/Source/PixelStreaming2Servers/Private/WebSocketServerWrapper.cpp
+++ b/Engine/Plugins/Media/PixelStreaming2/Source/PixelStreaming2Servers/Private/WebSocketServerWrapper.cpp
@@ -9,8 +9,8 @@ namespace UE::PixelStreaming2Servers
 {
 
 	/*
-	* ------------- FWebSocketConnection -------------
-	*/
+	 * ------------- FWebSocketConnection -------------
+	 */
 
 	FThreadSafeCounter FWebSocketConnection::IdGenerator = FThreadSafeCounter(100);
 
@@ -44,15 +44,15 @@ namespace UE::PixelStreaming2Servers
 	{
 		// Convert FString into uint8 array.
 		FTCHARToUTF8 UTF8String(*Message);
-		
+
 		// Send the uint8 buffer
 		// Note: Due to how this socket connection is implemented, only binary messages are supported
-		return SocketConnection->Send((const uint8*) UTF8String.Get(), UTF8String.Length(), false);
+		return SocketConnection->Send((const uint8*)UTF8String.Get(), UTF8String.Length(), false);
 	}
 
 	void FWebSocketConnection::SetCallbacks()
 	{
-		if(SocketConnection)
+		if (SocketConnection)
 		{
 			SocketConnection->SetReceiveCallBack(OnPacketReceivedCallback);
 			SocketConnection->SetSocketClosedCallBack(OnClosedCallback);
@@ -60,13 +60,12 @@ namespace UE::PixelStreaming2Servers
 	}
 
 	/*
-	* ------------- FWebSocketServerWrapper -------------
-	*/
+	 * ------------- FWebSocketServerWrapper -------------
+	 */
 
 	FWebSocketServerWrapper::FWebSocketServerWrapper()
 		: bLaunched(false)
 	{
-
 	}
 
 	FWebSocketServerWrapper::~FWebSocketServerWrapper()
@@ -74,24 +73,26 @@ namespace UE::PixelStreaming2Servers
 		Stop();
 	}
 
-	void FWebSocketServerWrapper::EnableWebServer(TArray<FWebSocketHttpMount> InDirectoriesToServe)
+	void FWebSocketServerWrapper::EnableWebServer(TArray<FWebSocketHttpMount> InDirectoriesToServe, bool bInServeHttps, const FWebSocketServerCertificates& InCertificates)
 	{
 		bEnableWebServer = true;
 		DirectoriesToServe = InDirectoriesToServe;
+		bServeHttps = bInServeHttps;
+		Certificates = InCertificates;
 	}
 
 	bool FWebSocketServerWrapper::Launch(uint16 Port)
 	{
 		WSServer = FModuleManager::Get().LoadModuleChecked<IWebSocketNetworkingModule>(TEXT("WebSocketNetworking")).CreateServer();
-		
-		if(bEnableWebServer)
+
+		if (bEnableWebServer)
 		{
-			WSServer->EnableHTTPServer(DirectoriesToServe);
+			WSServer->EnableHTTPServer(DirectoriesToServe, bServeHttps, Certificates);
 		}
 
 		OnClientConnectedCallback.BindRaw(this, &FWebSocketServerWrapper::OnConnectionOpened);
 		bLaunched = WSServer->Init(Port, OnClientConnectedCallback);
-		if(!bLaunched)
+		if (!bLaunched)
 		{
 			UE_LOG(LogPixelStreaming2Servers, Error, TEXT("Failed to launch Websocket server at ws://127.0.0.1:%d"), Port);
 			WSServer.Reset();
@@ -116,7 +117,7 @@ namespace UE::PixelStreaming2Servers
 
 	bool FWebSocketServerWrapper::GetFirstConnection(uint16& OutConnectionId) const
 	{
-		if(Connections.Num() > 0)
+		if (Connections.Num() > 0)
 		{
 			OutConnectionId = Connections.CreateConstIterator().Key();
 			return true;
@@ -172,7 +173,7 @@ namespace UE::PixelStreaming2Servers
 			NamedConnections.Remove(*Key);
 		}
 
-		if(Connections.Contains(ConnectionId))
+		if (Connections.Contains(ConnectionId))
 		{
 			return Connections.Remove(ConnectionId) > 0;
 		}
@@ -185,7 +186,7 @@ namespace UE::PixelStreaming2Servers
 
 	bool FWebSocketServerWrapper::Send(uint16 ConnectionId, FString Message) const
 	{
-		if(Connections.Contains(ConnectionId))
+		if (Connections.Contains(ConnectionId))
 		{
 			return Connections[ConnectionId]->Send(Message);
 		}
@@ -207,7 +208,7 @@ namespace UE::PixelStreaming2Servers
 
 	void FWebSocketServerWrapper::OnConnectionOpened(INetworkingWebSocket* Socket)
 	{
-		if(!Socket)
+		if (!Socket)
 		{
 			UE_LOG(LogPixelStreaming2Servers, Error, TEXT("Websocket client connected with a null socket."));
 			return;
@@ -219,7 +220,7 @@ namespace UE::PixelStreaming2Servers
 
 		// Had a new client connect over websocket, store the connection with a unique ID.
 		TUniquePtr<FWebSocketConnection> Connection = MakeUnique<FWebSocketConnection>(Socket);
-		const uint16 Id = Connection->GetId();
+		const uint16					 Id = Connection->GetId();
 
 		// Bind to socket callbacks for messages/closed.
 		Connection->OnPacketReceivedCallback.BindRaw(this, &FWebSocketServerWrapper::OnPacketReceived, Id);
@@ -234,7 +235,7 @@ namespace UE::PixelStreaming2Servers
 
 	void FWebSocketServerWrapper::OnPacketReceived(void* Data, int32 Size, uint16 ConnectionId)
 	{
-		if(Size > 0)
+		if (Size > 0)
 		{
 			TArrayView<uint8> DataView = MakeArrayView(static_cast<uint8*>(Data), Size);
 			OnMessage.Broadcast(ConnectionId, DataView);
@@ -254,16 +255,16 @@ namespace UE::PixelStreaming2Servers
 
 	void FWebSocketServerWrapper::Tick(float DeltaTime)
 	{
-		if(!IsLaunched())
+		if (!IsLaunched())
 		{
 			return;
 		}
 
-		if(WSServer)
+		if (WSServer)
 		{
 			// Tick the websocket server
 			WSServer->Tick();
 		}
 	}
 
-} // UE::PixelStreaming2Servers
+} // namespace UE::PixelStreaming2Servers
diff --git a/Engine/Plugins/Media/PixelStreaming2/Source/PixelStreaming2Servers/Private/WebSocketServerWrapper.h b/Engine/Plugins/Media/PixelStreaming2/Source/PixelStreaming2Servers/Private/WebSocketServerWrapper.h
--- a/Engine/Plugins/Media/PixelStreaming2/Source/PixelStreaming2Servers/Private/WebSocketServerWrapper.h
+++ b/Engine/Plugins/Media/PixelStreaming2/Source/PixelStreaming2Servers/Private/WebSocketServerWrapper.h
@@ -20,29 +20,28 @@ namespace UE::PixelStreaming2Servers
 	public:
 		FWebSocketConnection(INetworkingWebSocket* InSocketConnection);
 		~FWebSocketConnection();
-		uint16 GetId() const;
+		uint16			GetId() const;
 		TArray<FString> GetUrlArgs() const;
-		bool Send(FString Message) const;
-		void SetCallbacks();
+		bool			Send(FString Message) const;
+		void			SetCallbacks();
 
 	public:
 		FWebSocketPacketReceivedCallBack OnPacketReceivedCallback;
-		FWebSocketInfoCallBack OnClosedCallback;
+		FWebSocketInfoCallBack			 OnClosedCallback;
 
 	private:
 		static FThreadSafeCounter IdGenerator;
-		uint16 Id;
-		TArray<FString> UrlArgs;
-		INetworkingWebSocket* SocketConnection;
+		uint16					  Id;
+		TArray<FString>			  UrlArgs;
+		INetworkingWebSocket*	  SocketConnection;
 	};
 
-
 	class FWebSocketServerWrapper : public FTickableGameObject
 	{
 	public:
 		FWebSocketServerWrapper();
 		virtual ~FWebSocketServerWrapper();
-		void EnableWebServer(TArray<FWebSocketHttpMount> InDirectoriesToServe);
+		void EnableWebServer(TArray<FWebSocketHttpMount> InDirectoriesToServe, bool bServeHttps, const FWebSocketServerCertificates& Certificates);
 		bool Launch(uint16 Port);
 		void Stop();
 		bool IsLaunched() const;
@@ -53,14 +52,14 @@ namespace UE::PixelStreaming2Servers
 		 * Note: This message is sent as binary using the websocket protocol because of the WS implementation that is used.
 		 * @return True if the message able to be sent.
 		 */
-		bool Send(uint16 ConnectionId, FString Message) const;
-		bool Send(const FString& ConnectionName, FString Message) const;
-		bool GetFirstConnection(uint16& OutConnectionId) const;
+		bool											Send(uint16 ConnectionId, FString Message) const;
+		bool											Send(const FString& ConnectionName, FString Message) const;
+		bool											GetFirstConnection(uint16& OutConnectionId) const;
 		TMap<uint16, TUniquePtr<FWebSocketConnection>>& GetConnections() { return Connections; }
 
-		void NameConnection(uint16 ConnectionId, const FString& Name);
-		void RemoveName(const FString& Name);
-		bool GetNamedConnection(const FString& Name, uint16& OutConnectionId) const;
+		void			NameConnection(uint16 ConnectionId, const FString& Name);
+		void			RemoveName(const FString& Name);
+		bool			GetNamedConnection(const FString& Name, uint16& OutConnectionId) const;
 		TArray<FString> GetConnectionNames() const;
 
 		/* Begin FTickableGameObject */
@@ -68,7 +67,7 @@ namespace UE::PixelStreaming2Servers
 		virtual bool IsTickableInEditor() const { return true; }
 		virtual void Tick(float DeltaTime) override;
 		virtual bool IsAllowedToTick() const { return true; }
-		TStatId GetStatId() const { RETURN_QUICK_DECLARE_CYCLE_STAT(FWebSocketServerWrapper, STATGROUP_Tickables); }
+		TStatId		 GetStatId() const { RETURN_QUICK_DECLARE_CYCLE_STAT(FWebSocketServerWrapper, STATGROUP_Tickables); }
 		/* End FTickableGameObject */
 
 		int32 Count() const { return Connections.Num(); }
@@ -89,13 +88,15 @@ namespace UE::PixelStreaming2Servers
 		virtual void OnConnectionClosed(uint16 ConnectionId);
 
 	private:
-		FThreadSafeBool bLaunched;
-		TUniquePtr<IWebSocketServer> WSServer;
-		FWebSocketClientConnectedCallBack OnClientConnectedCallback;
+		FThreadSafeBool								   bLaunched;
+		TUniquePtr<IWebSocketServer>				   WSServer;
+		FWebSocketClientConnectedCallBack			   OnClientConnectedCallback;
 		TMap<uint16, TUniquePtr<FWebSocketConnection>> Connections;
-		TMap<FString, uint16> NamedConnections;
-		bool bEnableWebServer = false;
-		TArray<FWebSocketHttpMount> DirectoriesToServe;
+		TMap<FString, uint16>						   NamedConnections;
+		bool										   bEnableWebServer = false;
+		bool										   bServeHttps = false;
+		FWebSocketServerCertificates				   Certificates;
+		TArray<FWebSocketHttpMount>					   DirectoriesToServe;
 	};
 
 } // namespace UE::PixelStreaming2Servers
