IPv6 Support for PixelStreaming2 signaling server
Add optional IPv6 support to WebSocketsNetworking
Enable dual-stack listening when hosting PS from editor
#Synced-CL 52610047

diff --git a/Engine/Plugins/Experimental/WebSocketNetworking/Source/WebSocketNetworking/Private/WebSocketServer.cpp b/Engine/Plugins/Experimental/WebSocketNetworking/Source/WebSocketNetworking/Private/WebSocketServer.cpp
--- a/Engine/Plugins/Experimental/WebSocketNetworking/Source/WebSocketNetworking/Private/WebSocketServer.cpp
+++ b/Engine/Plugins/Experimental/WebSocketNetworking/Source/WebSocketNetworking/Private/WebSocketServer.cpp
@@ -165,7 +165,25 @@ bool FWebSocketServer::Init(uint32 Port, FWebSocketClientConnectedCallBack CallB
 		Info.ssl_cert_filepath = Certificates.GetCertificateFilePath();
 		Info.ssl_private_key_filepath = Certificates.GetPrivateKeyFilePath();
 	}
-	Info.options |= LWS_SERVER_OPTION_DISABLE_IPV6;
+	const bool bUseIPv6 = EnumHasAnyFlags(NetworkProtocol, IWebSocketServer::ENetworkProtocol::IPv6);
+	const bool bUseIPv4 = EnumHasAnyFlags(NetworkProtocol, IWebSocketServer::ENetworkProtocol::IPv4);
+
+	if (!bUseIPv6)
+	{
+		Info.options |= LWS_SERVER_OPTION_DISABLE_IPV6;
+	}
+	else
+	{
+		// LWS_SERVER_OPTION_IPV6_V6ONLY_MODIFY tells lws to set IPV6_V6ONLY explicitly.
+		Info.options |= LWS_SERVER_OPTION_IPV6_V6ONLY_MODIFY;
+		if (!bUseIPv4)
+		{
+			// IPv6-only: set IPV6_V6ONLY=1 to reject IPv4 connections.
+			Info.options |= LWS_SERVER_OPTION_IPV6_V6ONLY_VALUE;
+		}
+		// else: dual-stack (IPV6_V6ONLY=0) — IPv4 connections accepted via ::ffff: mapped addresses.
+		// This is required on Windows where the default is IPV6_V6ONLY=1.
+	}
 
 	if (bEnableHttp && LwsHttpMounts != NULL)
 	{
@@ -360,7 +378,7 @@ static int unreal_networking_server(
 					}
 				}
 
-				constexpr int32 IPAddressLength = 16;
+				constexpr int32 IPAddressLength = 46; // INET6_ADDRSTRLEN — large enough for IPv6
 				char			IPAddress[IPAddressLength];
 
 				lws_get_peer_simple(Wsi, IPAddress, IPAddressLength);
diff --git a/Engine/Plugins/Experimental/WebSocketNetworking/Source/WebSocketNetworking/Private/WebSocketServer.h b/Engine/Plugins/Experimental/WebSocketNetworking/Source/WebSocketNetworking/Private/WebSocketServer.h
--- a/Engine/Plugins/Experimental/WebSocketNetworking/Source/WebSocketNetworking/Private/WebSocketServer.h
+++ b/Engine/Plugins/Experimental/WebSocketNetworking/Source/WebSocketNetworking/Private/WebSocketServer.h
@@ -16,6 +16,7 @@ public:
 	virtual void	EnableHTTPServer(TArray<FWebSocketHttpMount> DirectoriesToServe, bool bSecure, const FWebSocketServerCertificates& Certificates) override;
 	virtual bool	Init(uint32 Port, FWebSocketClientConnectedCallBack, FString BindAddress = TEXT("")) override;
 	virtual void	SetFilterConnectionCallback(FWebSocketFilterConnectionCallback InFilterConnectionCallback) override;
+	virtual void	SetNetworkProtocol(IWebSocketServer::ENetworkProtocol Protocol) override { NetworkProtocol = Protocol; }
 	virtual void	Tick() override;
 	virtual FString Info() override;
 	//~ End IWebSocketServer interface
@@ -42,6 +43,8 @@ public:
 private:
 	bool bEnableHttp = false;
 
+	IWebSocketServer::ENetworkProtocol NetworkProtocol = IWebSocketServer::ENetworkProtocol::IPv4;
+
 	bool bSecure = false;
 
 	FWebSocketServerCertificates Certificates;
diff --git a/Engine/Plugins/Experimental/WebSocketNetworking/Source/WebSocketNetworking/Public/IWebSocketServer.h b/Engine/Plugins/Experimental/WebSocketNetworking/Source/WebSocketNetworking/Public/IWebSocketServer.h
--- a/Engine/Plugins/Experimental/WebSocketNetworking/Source/WebSocketNetworking/Public/IWebSocketServer.h
+++ b/Engine/Plugins/Experimental/WebSocketNetworking/Source/WebSocketNetworking/Public/IWebSocketServer.h
@@ -81,6 +81,14 @@ private:
 class IWebSocketServer
 {
 public:
+	/** Bitmask of network protocols the server should listen on. */
+	enum class ENetworkProtocol : uint8
+	{
+		None = 0,
+		IPv4 = 1 << 0,
+		IPv6 = 1 << 1,
+	};
+
 	virtual ~IWebSocketServer() = default;
 
 	/**
@@ -108,9 +116,20 @@ public:
 	 */
 	virtual void SetFilterConnectionCallback(FWebSocketFilterConnectionCallback InFilterConnectionCallback) = 0;
 
+	/**
+	 * Set which network protocols the server should listen on.
+	 * Must be called before Init(). Default: IPv4 only.
+	 *   IPv4       - listen on 0.0.0.0 (IPv4 only)
+	 *   IPv6       - listen on :: (IPv6 only, IPV6_V6ONLY=1)
+	 *   IPv4|IPv6  - listen on :: (dual-stack, IPV6_V6ONLY=0)
+	 */
+	virtual void SetNetworkProtocol(ENetworkProtocol Protocol) {}
+
 	/** Tick the server. */
 	virtual void Tick() = 0;
 
 	/** Describe this libwebsocket server */
 	virtual FString Info() = 0;
 };
+
+ENUM_CLASS_FLAGS(IWebSocketServer::ENetworkProtocol);
diff --git a/Engine/Plugins/Media/PixelStreaming2/Source/PixelStreaming2Editor/Private/PixelStreaming2EditorModule.cpp b/Engine/Plugins/Media/PixelStreaming2/Source/PixelStreaming2Editor/Private/PixelStreaming2EditorModule.cpp
--- a/Engine/Plugins/Media/PixelStreaming2/Source/PixelStreaming2Editor/Private/PixelStreaming2EditorModule.cpp
+++ b/Engine/Plugins/Media/PixelStreaming2/Source/PixelStreaming2Editor/Private/PixelStreaming2EditorModule.cpp
@@ -357,7 +357,7 @@ namespace UE::EditorPixelStreaming2
 		LaunchArgs.ReconnectionIntervalSeconds = 2.0f;
 
 		FString CertificateArgs = FString::Printf(TEXT("--CertificatePath=%s --PrivateKeyPath=%s"), *SSLCertificatePath, *SSLPrivateKeyPath);
-		LaunchArgs.ProcessArgs = FString::Printf(TEXT("--HttpPort=%d --StreamerPort=%d --ServeHttps=%s %s"), ViewerPort, StreamerPort, bServeHttps ? TEXT("true") : TEXT("false"), bServeHttps ? *CertificateArgs : TEXT(""));
+		LaunchArgs.ProcessArgs = FString::Printf(TEXT("--HttpPort=%d --StreamerPort=%d --ServeHttps=%s --EnableIPv6=true %s"), ViewerPort, StreamerPort, bServeHttps ? TEXT("true") : TEXT("false"), bServeHttps ? *CertificateArgs : TEXT(""));
 
 		FString AllowedOrigins = UPixelStreaming2PluginSettings::CVarAllowedOrigins.GetValueOnAnyThread();
 		if (!AllowedOrigins.IsEmpty())
diff --git a/Engine/Plugins/Media/PixelStreaming2/Source/PixelStreaming2Servers/Private/SignallingServer.cpp b/Engine/Plugins/Media/PixelStreaming2/Source/PixelStreaming2Servers/Private/SignallingServer.cpp
--- a/Engine/Plugins/Media/PixelStreaming2/Source/PixelStreaming2Servers/Private/SignallingServer.cpp
+++ b/Engine/Plugins/Media/PixelStreaming2/Source/PixelStreaming2Servers/Private/SignallingServer.cpp
@@ -145,10 +145,19 @@ namespace UE::PixelStreaming2Servers
 			UE_LOGFMT(LogPixelStreaming2Servers, Log, "CORS filtering enabled. Allowed origins: {0}", AllowedOriginsArg);
 		}
 
+		// Parse IPv6 setting — players WS gets dual-stack, streamers WS stays IPv4-only
+		FString EnableIPv6Str = Utils::QueryOrSetProcessArgs(InLaunchArgs, TEXT("--EnableIPv6="), TEXT("false"));
+		IWebSocketServer::ENetworkProtocol PlayerProtocol = IWebSocketServer::ENetworkProtocol::IPv4;
+		if (EnableIPv6Str == TEXT("true"))
+		{
+			PlayerProtocol = IWebSocketServer::ENetworkProtocol::IPv4 | IWebSocketServer::ENetworkProtocol::IPv6;
+		}
+
 		/*
 		 * --------------- Streamers websocket server ---------------
 		 */
 		StreamersWS = MakeUnique<FWebSocketServerWrapper>();
+		// Do not need to enable IPv6 on StreamerWS as it is always local
 		if (AllowedOrigins.Num() > 0)
 		{
 			FWebSocketFilterConnectionCallback FilterCallback;
@@ -188,6 +197,7 @@ namespace UE::PixelStreaming2Servers
 		 * --------------- Players websocket server ---------------
 		 */
 		PlayersWS = MakeUnique<FWebSocketServerWrapper>();
+		PlayersWS->SetNetworkProtocol(PlayerProtocol);
 		if (AllowedOrigins.Num() > 0)
 		{
 			FWebSocketFilterConnectionCallback FilterCallback;
diff --git a/Engine/Plugins/Media/PixelStreaming2/Source/PixelStreaming2Servers/Private/WebSocketServerWrapper.cpp b/Engine/Plugins/Media/PixelStreaming2/Source/PixelStreaming2Servers/Private/WebSocketServerWrapper.cpp
--- a/Engine/Plugins/Media/PixelStreaming2/Source/PixelStreaming2Servers/Private/WebSocketServerWrapper.cpp
+++ b/Engine/Plugins/Media/PixelStreaming2/Source/PixelStreaming2Servers/Private/WebSocketServerWrapper.cpp
@@ -107,17 +107,22 @@ namespace UE::PixelStreaming2Servers
 			WSServer->SetFilterConnectionCallback(FilterConnectionCallback);
 		}
 
+		WSServer->SetNetworkProtocol(NetworkProtocol);
+
 		OnClientConnectedCallback.BindRaw(this, &FWebSocketServerWrapper::OnConnectionOpened);
 		bLaunched = WSServer->Init(Port, OnClientConnectedCallback);
 		if (!bLaunched)
 		{
-			UE_LOG(LogPixelStreaming2Servers, Error, TEXT("Failed to launch Websocket server at ws://127.0.0.1:%d"), Port);
+			UE_LOG(LogPixelStreaming2Servers, Error, TEXT("Failed to launch Websocket server on port %d"), Port);
 			WSServer.Reset();
 			OnClientConnectedCallback.Unbind();
 		}
 		else
 		{
-			UE_LOG(LogPixelStreaming2Servers, Log, TEXT("Started Websocket server at ws://127.0.0.1:%d"), Port);
+			const bool bIPv4 = EnumHasAnyFlags(NetworkProtocol, IWebSocketServer::ENetworkProtocol::IPv4);
+			const bool bIPv6 = EnumHasAnyFlags(NetworkProtocol, IWebSocketServer::ENetworkProtocol::IPv6);
+			const TCHAR* ProtocolStr = (bIPv4 && bIPv6) ? TEXT("IPv4+IPv6") : bIPv6 ? TEXT("IPv6") : TEXT("IPv4");
+			UE_LOG(LogPixelStreaming2Servers, Log, TEXT("Started Websocket server on port %d (%s)"), Port, ProtocolStr);
 		}
 		return bLaunched;
 	}
diff --git a/Engine/Plugins/Media/PixelStreaming2/Source/PixelStreaming2Servers/Private/WebSocketServerWrapper.h b/Engine/Plugins/Media/PixelStreaming2/Source/PixelStreaming2Servers/Private/WebSocketServerWrapper.h
--- a/Engine/Plugins/Media/PixelStreaming2/Source/PixelStreaming2Servers/Private/WebSocketServerWrapper.h
+++ b/Engine/Plugins/Media/PixelStreaming2/Source/PixelStreaming2Servers/Private/WebSocketServerWrapper.h
@@ -38,6 +38,7 @@ namespace UE::PixelStreaming2Servers
 		virtual ~FWebSocketServerWrapper();
 		void EnableWebServer(TArray<FWebSocketHttpMount> InDirectoriesToServe, bool bServeHttps, const FWebSocketServerCertificates& Certificates);
 		void SetFilterConnectionCallback(FWebSocketFilterConnectionCallback InFilterConnectionCallback);
+		void SetNetworkProtocol(IWebSocketServer::ENetworkProtocol Protocol) { NetworkProtocol = Protocol; }
 		bool Launch(uint16 Port);
 		void Stop();
 		bool IsLaunched() const;
@@ -90,6 +91,7 @@ namespace UE::PixelStreaming2Servers
 		TMap<uint16, TUniquePtr<FWebSocketConnection>> Connections;
 		TMap<FString, uint16>						   NamedConnections;
 		bool										   bEnableWebServer = false;
+		IWebSocketServer::ENetworkProtocol			   NetworkProtocol = IWebSocketServer::ENetworkProtocol::IPv4;
 		bool										   bServeHttps = false;
 		FWebSocketServerCertificates				   Certificates;
 		TArray<FWebSocketHttpMount>					   DirectoriesToServe;
